Remove dependency on editor style from Slate UGS area

diff --git a/Engine/Source/Programs/SlateUGS/Private/SUnrealGameSyncWindow.cpp b/Engine/Source/Programs/SlateUGS/Private/SUnrealGameSyncWindow.cpp
--- a/Engine/Source/Programs/SlateUGS/Private/SUnrealGameSyncWindow.cpp
+++ b/Engine/Source/Programs/SlateUGS/Private/SUnrealGameSyncWindow.cpp
@@ -12,7 +12,6 @@
 #include "Widgets/Colors/SSimpleGradient.h"
 
 #include "Styling/AppStyle.h"
-#include "EditorStyleSet.h"
 
 #define LOCTEXT_NAMESPACE "UGSWindow"
 
diff --git a/Engine/Source/Programs/SlateUGS/SlateUGS.Build.cs b/Engine/Source/Programs/SlateUGS/SlateUGS.Build.cs
--- a/Engine/Source/Programs/SlateUGS/SlateUGS.Build.cs
+++ b/Engine/Source/Programs/SlateUGS/SlateUGS.Build.cs
@@ -18,7 +18,6 @@ public class SlateUGS : ModuleRules
 				"SlateCore",
 				"StandaloneRenderer",
 				"ToolWidgets",
-				"EditorStyle",
 			}
 		);
 
diff --git a/Engine/Source/Programs/SlateUGS/SlateUGS.Target.cs b/Engine/Source/Programs/SlateUGS/SlateUGS.Target.cs
--- a/Engine/Source/Programs/SlateUGS/SlateUGS.Target.cs
+++ b/Engine/Source/Programs/SlateUGS/SlateUGS.Target.cs
@@ -12,7 +12,6 @@ public class SlateUGSTarget : TargetRules
 		LinkType = TargetLinkType.Monolithic;
 
 		LaunchModuleName = "SlateUGS";
-		ExtraModuleNames.Add("EditorStyle");
 
 		bBuildDeveloperTools = false;
 
