Start of the UGSCore area. This was an original port done by Ben.M of UGS C# -> C++ at CL 3554544 (yeah its quite old and will be missing many newer UGS features!)

diff --git a/Engine/Source/Programs/SlateUGS/Private/UGSCore/BuildStep.cpp b/Engine/Source/Programs/SlateUGS/Private/UGSCore/BuildStep.cpp
new file mode 100644
--- /dev/null
+++ b/Engine/Source/Programs/SlateUGS/Private/UGSCore/BuildStep.cpp
@@ -0,0 +1,154 @@
+// Copyright 1998-2017 Epic Games, Inc. All Rights Reserved.
+
+#include "BuildStep.h"
+
+const TCHAR* FBuildStep::UniqueIdKey = TEXT("UniqueId");
+
+bool TryParse(const TCHAR* Text, EBuildStepType& OutType)
+{
+	if (FCString::Stricmp(Text, TEXT("Compile")) == 0)
+	{
+		OutType = EBuildStepType::Compile;
+		return true;
+	}
+	if (FCString::Stricmp(Text, TEXT("Cook")) == 0)
+	{
+		OutType = EBuildStepType::Cook;
+		return true;
+	}
+	if (FCString::Stricmp(Text, TEXT("Other")) == 0)
+	{
+		OutType = EBuildStepType::Other;
+		return true;
+	}
+	return false;
+}
+
+const TCHAR* ToString(EBuildStepType Type)
+{
+	switch(Type)
+	{
+	case EBuildStepType::Compile:
+		return TEXT("Compile");
+	case EBuildStepType::Cook:
+		return TEXT("Cook");
+	case EBuildStepType::Other:
+		return TEXT("Other");
+	default:
+		return TEXT("<Invalid>");
+	}
+}
+
+FBuildStep::FBuildStep(const FGuid& InUniqueId, int InOrderIndex, const FString& InDescription, const FString& InStatusText, int InEstimatedDuration, const FString& InTarget, const FString& InPlatform, const FString& InConfiguration, const FString& InArguments, bool bInSyncDefault)
+	: UniqueId(InUniqueId)
+	, OrderIndex(InOrderIndex)
+	, Description(InDescription)
+	, StatusText(InStatusText)
+	, EstimatedDuration(InEstimatedDuration)
+	, Type(EBuildStepType::Compile)
+	, Target(InTarget)
+	, Platform(InPlatform)
+	, Configuration(InConfiguration)
+	, Arguments(InArguments)
+	, bUseLogWindow(false)
+	, bNormalSync(bInSyncDefault)
+	, bScheduledSync(bInSyncDefault)
+	, bShowAsTool(false)
+{
+}
+
+FBuildStep::FBuildStep(const FCustomConfigObject& Object)
+{
+	if(!Object.TryGetValue(UniqueIdKey, UniqueId))
+	{
+		UniqueId = FGuid::NewGuid();
+	}
+	if(!Object.TryGetValue(TEXT("OrderIndex"), OrderIndex))
+	{
+		OrderIndex = -1;
+	}
+
+	Description = Object.GetValueOrDefault(TEXT("Description"), TEXT("Untitled"));
+	StatusText = Object.GetValueOrDefault(TEXT("StatusText"), TEXT("Untitled"));
+
+	if(!Object.TryGetValue(TEXT("EstimatedDuration"), EstimatedDuration) || EstimatedDuration < 1)
+	{
+		EstimatedDuration = 1;
+	}
+
+	if(!TryParse(Object.GetValueOrDefault(TEXT("Type"), TEXT("")), Type))
+	{
+		Type = EBuildStepType::Other;
+	}
+
+	Object.TryGetValue(TEXT("Target"), Target);
+	Object.TryGetValue(TEXT("Platform"), Platform);
+	Object.TryGetValue(TEXT("Configuration"), Configuration);
+	Object.TryGetValue(TEXT("FileName"), FileName);
+	Object.TryGetValue(TEXT("Arguments"), Arguments);
+
+	bUseLogWindow = Object.GetValueOrDefault(TEXT("bUseLogWindow"), true);
+	bNormalSync = Object.GetValueOrDefault(TEXT("bNormalSync"), true);
+	bScheduledSync = Object.GetValueOrDefault(TEXT("bScheduledSync"), true);
+	bShowAsTool = Object.GetValueOrDefault(TEXT("bShowAsTool"), false);
+}
+
+void FBuildStep::MergeBuildStepObjects(TMap<FGuid, FCustomConfigObject>& BuildStepObjects, const TArray<FCustomConfigObject>& ModifyObjects)
+{
+	for(const FCustomConfigObject& ModifyObject : ModifyObjects)
+	{
+		FGuid UniqueId;
+		if(ModifyObject.TryGetValue(UniqueIdKey, UniqueId))
+		{
+			FCustomConfigObject CombinedObject = ModifyObject;
+
+			FCustomConfigObject* DefaultObject = BuildStepObjects.Find(UniqueId);
+			if(DefaultObject != nullptr)
+			{
+				CombinedObject.SetDefaults(*DefaultObject);
+			}
+
+			BuildStepObjects[UniqueId] = CombinedObject;
+		}
+	}
+}
+
+FCustomConfigObject FBuildStep::ToConfigObject() const
+{
+	FCustomConfigObject Result;
+	Result.SetValue(TEXT("UniqueId"), UniqueId);
+	Result.SetValue(TEXT("Description"), *Description);
+	Result.SetValue(TEXT("StatusText"), *StatusText);
+	Result.SetValue(TEXT("EstimatedDuration"), EstimatedDuration);
+	Result.SetValue(TEXT("Type"), ToString(Type));
+	switch(Type)
+	{
+	case EBuildStepType::Compile:
+		Result.SetValue(TEXT("Target"), *Target);
+		Result.SetValue(TEXT("Platform"), *Platform);
+		Result.SetValue(TEXT("Configuration"), *Configuration);
+		Result.SetValue(TEXT("Arguments"), *Arguments);
+		break;
+	case EBuildStepType::Cook:
+		Result.SetValue(TEXT("FileName"), *FileName);
+		break;
+	case EBuildStepType::Other:
+		Result.SetValue(TEXT("FileName"), *FileName);
+		Result.SetValue(TEXT("Arguments"), *Arguments);
+		Result.SetValue(TEXT("bUseLogWindow"), bUseLogWindow);
+		break;
+	}
+	Result.SetValue(TEXT("OrderIndex"), OrderIndex);
+	Result.SetValue(TEXT("bNormalSync"), bNormalSync);
+	Result.SetValue(TEXT("bScheduledSync"), bScheduledSync);
+	Result.SetValue(TEXT("bShowAsTool"), bShowAsTool);
+	return Result;
+}
+
+FCustomConfigObject FBuildStep::ToConfigObject(const FCustomConfigObject* DefaultObject) const
+{
+	FCustomConfigObject Result;
+	Result.SetValue(UniqueIdKey, UniqueId);
+	Result.AddOverrides(ToConfigObject(), DefaultObject);
+	return Result;
+}
diff --git a/Engine/Source/Programs/SlateUGS/Private/UGSCore/BuildStep.h b/Engine/Source/Programs/SlateUGS/Private/UGSCore/BuildStep.h
new file mode 100644
--- /dev/null
+++ b/Engine/Source/Programs/SlateUGS/Private/UGSCore/BuildStep.h
@@ -0,0 +1,46 @@
+// Copyright 1998-2017 Epic Games, Inc. All Rights Reserved.
+
+#pragma once
+
+#include "CoreMinimal.h"
+#include "Misc/Guid.h"
+#include "CustomConfigFile.h"
+
+enum class EBuildStepType
+{
+	Compile,
+	Cook,
+	Other,
+};
+
+bool TryParse(const TCHAR* Text, EBuildStepType& OutType);
+const TCHAR* ToString(EBuildStepType Type);
+
+struct FBuildStep
+{
+	static const TCHAR* UniqueIdKey;
+
+	FGuid UniqueId;
+	int OrderIndex;
+	FString Description;
+	FString StatusText;
+	int EstimatedDuration;
+	EBuildStepType Type;
+	FString Target;
+	FString Platform;
+	FString Configuration;
+	FString FileName;
+	FString Arguments;
+	bool bUseLogWindow;
+	bool bNormalSync;
+	bool bScheduledSync;
+	bool bShowAsTool;
+
+	FBuildStep(const FGuid& InUniqueId, int InOrderIndex, const FString& InDescription, const FString& InStatusText, int InEstimatedDuration, const FString& InTarget, const FString& InPlatform, const FString& InConfiguration, const FString& InArguments, bool bInSyncDefault);
+	FBuildStep(const FCustomConfigObject& Object);
+
+	static void MergeBuildStepObjects(TMap<FGuid, FCustomConfigObject>& BuildStepObjects, const TArray<FCustomConfigObject>& ModifyObjects);
+
+	FCustomConfigObject ToConfigObject() const;
+	FCustomConfigObject ToConfigObject(const FCustomConfigObject* DefaultObject) const;
+};
diff --git a/Engine/Source/Programs/SlateUGS/Private/UGSCore/CustomConfigFile.cpp b/Engine/Source/Programs/SlateUGS/Private/UGSCore/CustomConfigFile.cpp
new file mode 100644
--- /dev/null
+++ b/Engine/Source/Programs/SlateUGS/Private/UGSCore/CustomConfigFile.cpp
@@ -0,0 +1,736 @@
+// Copyright 1998-2017 Epic Games, Inc. All Rights Reserved.
+
+#include "CustomConfigFile.h"
+#include "Misc/FileHelper.h"
+#include "Utility.h"
+
+static const TCHAR* ConfigSeparatorCharacters = TEXT("(),= \t\"");
+
+//// FCustomConfigObject ////
+
+FCustomConfigObject::FCustomConfigObject()
+{
+}
+
+FCustomConfigObject::FCustomConfigObject(const TCHAR* Text)
+{
+	ParseConfigString(Text);
+}
+
+FCustomConfigObject::FCustomConfigObject(const FCustomConfigObject& BaseObject, const TCHAR* Text)
+	: Pairs(BaseObject.Pairs)
+{
+	ParseConfigString(Text);
+}
+
+void FCustomConfigObject::ParseConfigString(const TCHAR* Text)
+{
+	int Idx = 0;
+	if(ParseConfigToken(Text, Idx) == "(")
+	{
+		while(Text[Idx] != 0)
+		{
+			// Read the next key/value pair
+			FString Key = ParseConfigToken(Text, Idx);
+			if(ParseConfigToken(Text, Idx) == "=")
+			{
+				FString Value = ParseConfigToken(Text, Idx);
+				SetValue(*Key, *Value);
+			}
+
+			// Check for the end of the list, or a comma before the next pair
+			for(;;)
+			{
+				FString Token = ParseConfigToken(Text, Idx);
+				if(Token == ",")
+				{
+					break;
+				}
+				if(Token == ")" || Token.Len() == 0)
+				{
+					return;
+				}
+			}
+		}
+	}
+}
+
+FString FCustomConfigObject::ParseConfigToken(const TCHAR* Text, int& Idx)
+{
+	// Skip whitespace
+	while(Text[Idx] != 0 && FChar::IsWhitespace(Text[Idx]))
+	{
+		Idx++;
+	}
+	if(Text[Idx] == 0)
+	{
+		return FString();
+	}
+
+	// Read the token
+	if(Text[Idx] == '\"')
+	{
+		FString Token;
+		while(Text[++Idx] != 0)
+		{
+			if(Text[Idx] == '\"')
+			{
+				Idx++;
+				break;
+			}
+			if(Text[Idx] == '\\' && Text[Idx + 1] != 0)
+			{
+				Idx++;
+			}
+			Token += Text[Idx];
+		}
+		return Token;
+	}
+	else if(FCString::Strchr(ConfigSeparatorCharacters, Text[Idx]) != nullptr)
+	{
+		FString Token;
+		Token.AppendChar(Text[Idx++]);
+		return Token;
+	}
+	else
+	{
+		int StartIdx = Idx;
+		while(Text[Idx] != 0 && FCString::Strchr(ConfigSeparatorCharacters, Text[Idx]) == nullptr)
+		{
+			Idx++;
+		}
+		return FString(Idx - StartIdx, Text + StartIdx);
+	}
+}
+
+bool FCustomConfigObject::HasKey(const TCHAR* Key) const
+{
+	const TCHAR* CurrentValue;
+	return TryGetValue(Key, CurrentValue);
+}
+
+bool FCustomConfigObject::HasValue(const TCHAR* Key, const TCHAR* Value) const
+{
+	const TCHAR* CurrentValue;
+	return TryGetValue(Key, CurrentValue) && FCString::Stricmp(CurrentValue, Value) == 0;
+}
+
+bool FCustomConfigObject::TryGetValue(const TCHAR* Key, const TCHAR*& OutValue) const
+{
+	for(int Idx = 0; Idx < Pairs.Num(); Idx++)
+	{
+		if(Pairs[Idx].Key == Key)
+		{
+			OutValue = *Pairs[Idx].Value;
+			return true;
+		}
+	}
+	return false;
+}
+
+bool FCustomConfigObject::TryGetValue(const TCHAR* Key, FString& OutValue) const
+{
+	const TCHAR* Value;
+	if(TryGetValue(Key, Value))
+	{
+		OutValue = Value;
+		return true;
+	}
+	return false;
+}
+
+bool FCustomConfigObject::TryGetValue(const TCHAR* Key, FGuid& OutValue) const
+{
+	const TCHAR* StringValue;
+	if(TryGetValue(Key, StringValue))
+	{
+		FGuid GuidValue;
+		if(FGuid::Parse(StringValue, GuidValue))
+		{
+			OutValue = GuidValue;
+			return true;
+		}
+	}
+	return false;
+}
+
+bool FCustomConfigObject::TryGetValue(const TCHAR* Key, int& OutValue) const
+{
+	const TCHAR* StringValue;
+	if(TryGetValue(Key, StringValue))
+	{
+		TCHAR* EndPtr = nullptr;
+		int IntegerValue = FCString::Strtoi(StringValue, &EndPtr, 10);
+		if(*EndPtr == 0)
+		{
+			OutValue = IntegerValue;
+			return true;
+		}
+	}
+	return false;
+}
+
+bool FCustomConfigObject::TryGetValue(const TCHAR* Key, bool& OutValue) const
+{
+	const TCHAR* StringValue;
+	if(TryGetValue(Key, StringValue))
+	{
+		if(FCString::Stricmp(StringValue, TEXT("True")) == 0)
+		{
+			OutValue = true;
+			return true;
+		}
+		else if(FCString::Stricmp(StringValue, TEXT("False")) == 0)
+		{
+			OutValue = false;
+			return true;
+		}
+	}
+	return false;
+}
+
+const TCHAR* FCustomConfigObject::GetValueOrDefault(const TCHAR* Key, const TCHAR* DefaultValue) const
+{
+	const TCHAR* Value;
+	if(TryGetValue(Key, Value))
+	{
+		return Value;
+	}
+	else
+	{
+		return DefaultValue;
+	}
+}
+
+FGuid FCustomConfigObject::GetValueOrDefault(const TCHAR* Key, const FGuid& DefaultValue) const
+{
+	FGuid Value;
+	if(TryGetValue(Key, Value))
+	{
+		return Value;
+	}
+	else
+	{
+		return DefaultValue;
+	}
+}
+
+int FCustomConfigObject::GetValueOrDefault(const TCHAR* Key, int DefaultValue) const
+{
+	int Value;
+	if(TryGetValue(Key, Value))
+	{
+		return Value;
+	}
+	else
+	{
+		return DefaultValue;
+	}
+}
+
+bool FCustomConfigObject::GetValueOrDefault(const TCHAR* Key, bool DefaultValue) const
+{
+	bool Value;
+	if(TryGetValue(Key, Value))
+	{
+		return Value;
+	}
+	else
+	{
+		return DefaultValue;
+	}
+}
+
+void FCustomConfigObject::SetValue(const TCHAR* Key, const TCHAR* Value)
+{
+	for(int Idx = 0; Idx < Pairs.Num(); Idx++)
+	{
+		if(Pairs[Idx].Key == Key)
+		{
+			Pairs[Idx] = TTuple<FString,FString>(Key, Value);
+			return;
+		}
+	}
+	Pairs.Add(TTuple<FString,FString>(Key, Value));
+}
+
+void FCustomConfigObject::SetValue(const TCHAR* Key, const FGuid& Value)
+{
+	SetValue(Key, *Value.ToString());
+}
+
+void FCustomConfigObject::SetValue(const TCHAR* Key, int Value)
+{
+	SetValue(Key, *FString::Printf(TEXT("%d"), Value));
+}
+
+void FCustomConfigObject::SetValue(const TCHAR* Key, bool Value)
+{
+	SetValue(Key, Value? TEXT("True") : TEXT("False"));
+}
+
+void FCustomConfigObject::SetDefaults(const FCustomConfigObject& Other)
+{
+	for(const TTuple<FString, FString>& Pair : Other.Pairs)
+	{
+		if(!HasKey(*Pair.Key))
+		{
+			SetValue(*Pair.Key, *Pair.Value);
+		}
+	}
+}
+
+void FCustomConfigObject::AddOverrides(const FCustomConfigObject& Object, const FCustomConfigObject* DefaultObject)
+{
+	for(const TTuple<FString, FString>& Pair : Object.Pairs)
+	{
+		if(DefaultObject == nullptr || !DefaultObject->HasValue(*Pair.Key, *Pair.Value))
+		{
+			SetValue(*Pair.Key, *Pair.Value);
+		}
+	}
+}
+
+FString FCustomConfigObject::ToString(const FCustomConfigObject* BaseObject) const
+{
+	FString Result;
+	Result.Append("(");
+	for(const TTuple<FString, FString>& Pair : Pairs)
+	{
+		if(BaseObject == nullptr || !BaseObject->HasValue(*Pair.Key, *Pair.Value))
+		{
+			if(Result.Len() > 1)
+			{
+				Result.Append(", ");
+			}
+			Result.Append(Pair.Key);
+			Result.Append("=");
+			if(Pair.Value.Len() == 0)
+			{
+				Result.Append("\"\"");
+			}
+			else
+			{
+				Result.Append(FString(TEXT("\"")) + Pair.Value.Replace(TEXT("\\"), TEXT("\\\\")).Replace(TEXT("\""), TEXT("\\\"")) + FString(TEXT("\"")));
+			}
+		}
+	}
+	Result.Append(")");
+	return Result;
+}
+ 
+FString FCustomConfigObject::ToString() const
+{
+	return ToString(nullptr);
+}
+
+//// FCustomConfigSection ////
+
+FCustomConfigSection::FCustomConfigSection(FString&& InName)
+	: Name(InName)
+{
+}
+
+void FCustomConfigSection::Clear()
+{
+	Pairs.Empty();
+}
+
+void FCustomConfigSection::SetValue(const TCHAR* Key, int Value)
+{
+	Pairs.FindOrAdd(Key) = FString::Printf(TEXT("%d"), Value);
+}
+
+void FCustomConfigSection::SetValue(const TCHAR* Key, bool Value)
+{
+	Pairs.FindOrAdd(Key) = Value? "1" : "0";
+}
+
+void FCustomConfigSection::SetValue(const TCHAR* Key, const TCHAR* Value)
+{
+	if(Value == nullptr)
+	{
+		RemoveValue(Key);
+	}
+	else
+	{
+		Pairs.FindOrAdd(Key) = Value;
+	}
+}
+
+void FCustomConfigSection::SetValues(const TCHAR* Key, const TArray<FString>& Values)
+{
+	if(Values.Num() == 0)
+	{
+		RemoveValue(Key);
+	}
+	else
+	{
+		Pairs.FindOrAdd(Key) = FString::Join(Values, TEXT("\n"));
+	}
+}
+
+void FCustomConfigSection::SetValues(const TCHAR* Key, const TArray<FGuid>& Values)
+{
+	if(Values.Num() == 0)
+	{
+		RemoveValue(Key);
+	}
+	else
+	{
+		FString Value = Values[0].ToString();
+		for(int Idx = 1; Idx < Values.Num(); Idx++)
+		{
+			Value += "\n";
+			Value += Values[Idx].ToString();
+		}
+		Pairs.FindOrAdd(Key) = Value;
+	}
+}
+
+void FCustomConfigSection::SetValues(const TCHAR* Key, const TArray<int>& Values)
+{
+	if(Values.Num() == 0)
+	{
+		RemoveValue(Key);
+	}
+	else
+	{
+		FString Value = FString::Printf(TEXT("%d"), Values[0]);
+		for(int Idx = 1; Idx < Values.Num(); Idx++)
+		{
+			Value += TEXT("\n");
+			Value += FString::Printf(TEXT("%d"), Values[Idx]);
+		}
+		Pairs.FindOrAdd(Key) = Value;
+	}
+}
+
+void FCustomConfigSection::AppendValue(const TCHAR* Key, const TCHAR* Value)
+{
+	FString* CurrentValue = Pairs.Find(Key);
+	if(CurrentValue != nullptr)
+	{
+		(*CurrentValue) += FString::Printf(TEXT("\n%s"), Value);
+	}
+	else
+	{
+		Pairs.FindOrAdd(Key) = Value;
+	}
+}
+
+void FCustomConfigSection::RemoveValue(const TCHAR* Key)
+{
+	Pairs.Remove(Key);
+}
+
+int FCustomConfigSection::GetValue(const TCHAR* Key, int DefaultValue) const
+{
+	const TCHAR* ValueString = GetValue(Key);
+	if(ValueString != nullptr)
+	{
+		TCHAR* EndPtr = nullptr;
+		int Value = FCString::Strtoi(ValueString, &EndPtr, 10);
+		if(*EndPtr == 0)
+		{
+			return Value;
+		}
+	}
+	return DefaultValue;
+}
+
+bool FCustomConfigSection::GetValue(const TCHAR* Key, bool DefaultValue) const
+{
+	return GetValue(Key, DefaultValue? 1 : 0) != 0;
+}
+
+const TCHAR* FCustomConfigSection::GetValue(const TCHAR* Key, const TCHAR* DefaultValue) const
+{
+	const FString* Value = Pairs.Find(Key);
+	if(Value == nullptr)
+	{
+		return DefaultValue;
+	}
+	else
+	{
+		return **Value;
+	}
+}
+
+bool FCustomConfigSection::TryGetValue(const TCHAR* Key, FString& OutValue) const
+{
+	const TCHAR* ValueString = GetValue(Key);
+	if(ValueString != nullptr)
+	{
+		OutValue = ValueString;
+		return true;
+	}
+	return false;
+}
+
+bool FCustomConfigSection::TryGetValue(const TCHAR* Key, int32& OutValue) const
+{
+	const TCHAR* ValueString = GetValue(Key);
+	if(ValueString != nullptr)
+	{
+		TCHAR* EndPtr = nullptr;
+		int Value = FCString::Strtoi(ValueString, &EndPtr, 10);
+		if(*EndPtr == 0)
+		{
+			OutValue = Value;
+			return true;
+		}
+	}
+	return false;
+}
+
+bool FCustomConfigSection::TryGetValues(const TCHAR* Key, TArray<FString>& Values) const
+{
+	const FString* Value = Pairs.Find(Key);
+	if(Value != nullptr)
+	{
+		Value->ParseIntoArray(Values, TEXT("\n"), false);
+		return true;
+	}
+	return false;
+}
+
+bool FCustomConfigSection::TryGetValues(const TCHAR* Key, TArray<FGuid>& Values) const
+{
+	TArray<FString> StringValues;
+	if(TryGetValues(Key, StringValues))
+	{
+		for(const FString& StringValue : StringValues)
+		{
+			FGuid Guid;
+			if(FGuid::Parse(StringValue, Guid))
+			{
+				Values.Add(Guid);
+			}
+		}
+		return true;
+	}
+	return false;
+}
+
+bool FCustomConfigSection::TryGetValues(const TCHAR* Key, TArray<int>& Values) const
+{
+	TArray<FString> StringValues;
+	if(TryGetValues(Key, StringValues))
+	{
+		for(const FString& StringValue : StringValues)
+		{
+			int IntegerValue;
+			if(FUtility::TryParse(*StringValue, IntegerValue))
+			{
+				Values.Add(IntegerValue);
+			}
+		}
+		return true;
+	}
+	return false;
+}
+
+//// FCustomConfigFile ////
+
+FCustomConfigFile::FCustomConfigFile()
+{
+}
+
+void FCustomConfigFile::Load(const TCHAR* FileName)
+{
+	TArray<FString> Lines;
+	FFileHelper::LoadFileToStringArray(Lines, FileName);
+	Parse(Lines);
+}
+
+void FCustomConfigFile::Parse(const TArray<FString>& Lines)
+{
+	TSharedPtr<FCustomConfigSection> CurrentSection;
+	for(const FString& Line : Lines)
+	{
+		FString TrimLine = Line.TrimStartAndEnd();
+		if(!TrimLine.StartsWith(";"))
+		{
+			if(TrimLine.StartsWith("[") && TrimLine.EndsWith("]"))
+			{
+				FString SectionName = TrimLine.Mid(1, TrimLine.Len() - 2).TrimStartAndEnd();
+				CurrentSection = FindOrAddSection(*SectionName);
+			}
+			else if(CurrentSection.IsValid())
+			{
+				int EqualsIdx;
+				if(TrimLine.FindChar(TEXT('='), EqualsIdx))
+				{
+					FString Value = Line.Mid(EqualsIdx + 1).TrimStart();
+					if(TrimLine.StartsWith("+"))
+					{
+						FString Key = TrimLine.Mid(1, EqualsIdx - 1).TrimStartAndEnd();
+						CurrentSection->AppendValue(*Key, *Value);
+					}
+					else
+					{
+						FString Key = TrimLine.Mid(0, EqualsIdx).TrimEnd();
+						CurrentSection->SetValue(*Key, *Value);
+					}
+				}
+			}
+		}
+	}
+}
+
+void FCustomConfigFile::Save(const TCHAR* FileName)
+{
+	return;
+
+	TArray<FString> Lines;
+	for(int Idx = 0; Idx < Sections.Num(); Idx++)
+	{
+		Lines.Add(FString::Printf(TEXT("[%s]"), *Sections[Idx]->Name));
+		for(const TTuple<FString, FString>& Pair : Sections[Idx]->Pairs)
+		{
+			if(Pair.Value.Contains(TEXT("\n")))
+			{
+				TArray<FString> Values;
+				Pair.Value.ParseIntoArrayLines(Values, false);
+
+				for(const FString& Value : Values)
+				{
+					Lines.Add(FString::Printf(TEXT("+%s=%s"), *Pair.Key, *Value));
+				}
+			}
+			else
+			{
+				Lines.Add(FString::Printf(TEXT("%s=%s"), *Pair.Key, *Pair.Value));
+			}
+		}
+		if(Idx < Sections.Num() - 1)
+		{
+			Lines.Add(FString());
+		}
+	}
+}
+
+TSharedPtr<FCustomConfigSection> FCustomConfigFile::FindSection(const TCHAR* Name)
+{
+	for(const TSharedPtr<FCustomConfigSection>& Section : Sections)
+	{
+		if(Section->Name == Name)
+		{
+			return Section;
+		}
+	}
+	return nullptr;
+}
+
+TSharedPtr<const FCustomConfigSection> FCustomConfigFile::FindSection(const TCHAR* Name) const
+{
+	for(const TSharedPtr<FCustomConfigSection>& Section : Sections)
+	{
+		if(Section->Name == Name)
+		{
+			return Section;
+		}
+	}
+	return nullptr;
+}
+
+TSharedRef<FCustomConfigSection> FCustomConfigFile::FindOrAddSection(const TCHAR* Name)
+{
+	TSharedPtr<FCustomConfigSection> Section = FindSection(Name);
+	if(!Section.IsValid())
+	{
+		TSharedRef<FCustomConfigSection> NewSection = MakeShared<FCustomConfigSection>(Name);
+		Sections.Add(NewSection);
+		return NewSection;
+	}
+	return Section.ToSharedRef();
+}
+
+void FCustomConfigFile::SetValue(const TCHAR* Key, int Value)
+{
+	const TCHAR* Dot = FCString::Strchr(Key, TEXT('.'));
+	check(Dot != nullptr);
+	TSharedRef<FCustomConfigSection> Section = FindOrAddSection(*FString(Dot - Key, Key));
+	Section->SetValue(Dot + 1, Value);
+}
+
+void FCustomConfigFile::SetValue(const TCHAR* Key, bool Value)
+{
+	const TCHAR* Dot = FCString::Strchr(Key, TEXT('.'));
+	check(Dot != nullptr);
+	TSharedRef<FCustomConfigSection> Section = FindOrAddSection(*FString(Dot - Key, Key));
+	Section->SetValue(Dot + 1, Value);
+}
+
+void FCustomConfigFile::SetValue(const TCHAR* Key, const TCHAR* Value)
+{
+	const TCHAR* Dot = FCString::Strchr(Key, TEXT('.'));
+	check(Dot != nullptr);
+	TSharedRef<FCustomConfigSection> Section = FindOrAddSection(*FString(Dot - Key, Key));
+	Section->SetValue(Dot + 1, Value);
+}
+
+void FCustomConfigFile::SetValues(const TCHAR* Key, const TArray<FString>& Values)
+{
+	const TCHAR* Dot = FCString::Strchr(Key, TEXT('.'));
+	check(Dot != nullptr);
+	TSharedRef<FCustomConfigSection> Section = FindOrAddSection(*FString(Dot - Key, Key));
+	Section->SetValues(Dot + 1, Values);
+}
+
+bool FCustomConfigFile::GetValue(const TCHAR* Key, bool DefaultValue) const
+{
+	const TCHAR* Dot = FCString::Strchr(Key, TEXT('.'));
+	check(Dot != nullptr);
+	TSharedPtr<const FCustomConfigSection> Section = FindSection(*FString(Dot - Key, Key));
+	return Section.IsValid()? Section->GetValue(Dot + 1, DefaultValue) : DefaultValue;
+}
+
+int32 FCustomConfigFile::GetValue(const TCHAR* Key, int32 DefaultValue) const
+{
+	const TCHAR* Dot = FCString::Strchr(Key, TEXT('.'));
+	check(Dot != nullptr);
+	TSharedPtr<const FCustomConfigSection> Section = FindSection(*FString(Dot - Key, Key));
+	return Section.IsValid()? Section->GetValue(Dot + 1, DefaultValue) : DefaultValue;
+}
+
+const TCHAR* FCustomConfigFile::GetValue(const TCHAR* Key, const TCHAR* DefaultValue) const
+{
+	const TCHAR* Dot = FCString::Strchr(Key, TEXT('.'));
+	check(Dot != nullptr);
+	TSharedPtr<const FCustomConfigSection> Section = FindSection(*FString(Dot - Key, Key));
+	return Section.IsValid()? Section->GetValue(Dot + 1, DefaultValue) : DefaultValue;
+}
+
+bool FCustomConfigFile::TryGetValue(const TCHAR* Key, FString& OutValue) const
+{
+	const TCHAR* Dot = FCString::Strchr(Key, TEXT('.'));
+	check(Dot != nullptr);
+	TSharedPtr<const FCustomConfigSection> Section = FindSection(*FString(Dot - Key, Key));
+	return Section.IsValid()? Section->TryGetValue(Key, OutValue) : false;
+}
+
+bool FCustomConfigFile::TryGetValue(const TCHAR* Key, int32& OutValue) const
+{
+	const TCHAR* Dot = FCString::Strchr(Key, TEXT('.'));
+	check(Dot != nullptr);
+	TSharedPtr<const FCustomConfigSection> Section = FindSection(*FString(Dot - Key, Key));
+	return Section.IsValid()? Section->TryGetValue(Key, OutValue) : false;
+}
+
+bool FCustomConfigFile::TryGetValues(const TCHAR* Key, TArray<FString>& Values) const
+{
+	const TCHAR* Dot = FCString::Strchr(Key, TEXT('.'));
+	check(Dot != nullptr);
+	TSharedPtr<const FCustomConfigSection> Section = FindSection(*FString(Dot - Key, Key));
+	return Section.IsValid()? Section->TryGetValues(Dot + 1, Values) : false;
+}
+
+bool FCustomConfigFile::TryGetValues(const TCHAR* Key, TArray<FGuid> &Values) const
+{
+	const TCHAR* Dot = FCString::Strchr(Key, TEXT('.'));
+	check(Dot != nullptr);
+	TSharedPtr<const FCustomConfigSection> Section = FindSection(*FString(Dot - Key, Key));
+	return Section.IsValid()? Section->TryGetValues(Dot + 1, Values) : false;
+}
diff --git a/Engine/Source/Programs/SlateUGS/Private/UGSCore/CustomConfigFile.h b/Engine/Source/Programs/SlateUGS/Private/UGSCore/CustomConfigFile.h
new file mode 100644
--- /dev/null
+++ b/Engine/Source/Programs/SlateUGS/Private/UGSCore/CustomConfigFile.h
@@ -0,0 +1,98 @@
+// Copyright 1998-2017 Epic Games, Inc. All Rights Reserved.
+
+#pragma once
+
+#include "CoreMinimal.h"
+#include "Misc/Guid.h"
+#include "Templates/Tuple.h"
+
+struct FCustomConfigObject
+{
+	TArray<TTuple<FString, FString>> Pairs;
+
+	FCustomConfigObject();
+	FCustomConfigObject(const TCHAR* Text);
+	FCustomConfigObject(const FCustomConfigObject& BaseObject, const TCHAR* Text);
+
+	void ParseConfigString(const TCHAR* Text);
+	FString ParseConfigToken(const TCHAR* Text, int& Idx);
+
+	bool HasKey(const TCHAR* Key) const;
+	bool HasValue(const TCHAR* Key, const TCHAR* Value) const;
+
+	bool TryGetValue(const TCHAR* Key, const TCHAR*& OutValue) const;
+	bool TryGetValue(const TCHAR* Key, FString& OutValue) const;
+	bool TryGetValue(const TCHAR* Key, FGuid& OutValue) const;
+	bool TryGetValue(const TCHAR* Key, int& OutValue) const;
+	bool TryGetValue(const TCHAR* Key, bool& OutValue) const;
+
+	const TCHAR* GetValueOrDefault(const TCHAR* Key, const TCHAR* DefaultValue) const;
+	FGuid GetValueOrDefault(const TCHAR* Key, const FGuid& DefaultValue) const;
+	int GetValueOrDefault(const TCHAR* Key, int DefaultValue) const;
+	bool GetValueOrDefault(const TCHAR* Key, bool DefaultValue) const;
+
+	void SetValue(const TCHAR* Key, const TCHAR* Value);
+	void SetValue(const TCHAR* Key, const FGuid& Value);
+	void SetValue(const TCHAR* Key, int Value);
+	void SetValue(const TCHAR* Key, bool Value);
+
+	void SetDefaults(const FCustomConfigObject& Other);
+	void AddOverrides(const FCustomConfigObject& Object, const FCustomConfigObject* DefaultObject);
+
+	FString ToString(const FCustomConfigObject* BaseObject) const;
+	FString ToString() const;
+};
+
+struct FCustomConfigSection
+{
+	const FString Name;
+	TMap<FString, FString> Pairs;
+
+	FCustomConfigSection(FString&& InName);
+	void Clear();
+	void SetValue(const TCHAR* Key, int Value);
+	void SetValue(const TCHAR* Key, bool Value);
+	void SetValue(const TCHAR* Key, const TCHAR* Value);
+	void SetValues(const TCHAR* Key, const TArray<FString>& Values);
+	void SetValues(const TCHAR* Key, const TArray<FGuid>& Values);
+	void SetValues(const TCHAR* Key, const TArray<int>& Values);
+	void AppendValue(const TCHAR* Key, const TCHAR* Value);
+	void RemoveValue(const TCHAR* Key);
+	int GetValue(const TCHAR* Key, int DefaultValue) const;
+	bool GetValue(const TCHAR* Key, bool DefaultValue) const;
+	const TCHAR* GetValue(const TCHAR* Key, const TCHAR* DefaultValue = nullptr) const;
+
+	bool TryGetValue(const TCHAR* Key, FString& OutValue) const;
+	bool TryGetValue(const TCHAR* Key, int32& OutValue) const;
+	bool TryGetValues(const TCHAR* Key, TArray<FString>& Values) const;
+	bool TryGetValues(const TCHAR* Key, TArray<FGuid>& Values) const;
+	bool TryGetValues(const TCHAR* Key, TArray<int>& Values) const;
+};
+
+struct FCustomConfigFile
+{
+	TArray<TSharedPtr<FCustomConfigSection>> Sections;
+
+	FCustomConfigFile();
+	void Load(const TCHAR* FileName);
+	void Parse(const TArray<FString>& Lines);
+	void Save(const TCHAR* FileName);
+	TSharedPtr<FCustomConfigSection> FindSection(const TCHAR* Name);
+	TSharedPtr<const FCustomConfigSection> FindSection(const TCHAR* Name) const;
+	TSharedRef<FCustomConfigSection> FindOrAddSection(const TCHAR* Name);
+
+	void SetValue(const TCHAR* Key, int Value);
+	void SetValue(const TCHAR* Key, bool Value);
+	void SetValue(const TCHAR* Key, const TCHAR* Value);
+
+	void SetValues(const TCHAR* Key, const TArray<FString>& Values);
+
+	bool GetValue(const TCHAR* Key, bool DefaultValue) const;
+	int32 GetValue(const TCHAR* Key, int32 DefaultValue) const;
+	const TCHAR* GetValue(const TCHAR* Key, const TCHAR* DefaultValue) const;
+
+	bool TryGetValue(const TCHAR* Key, FString& OutValue) const;
+	bool TryGetValue(const TCHAR* Key, int32& OutValue) const;
+	bool TryGetValues(const TCHAR* Key, TArray<FString>& Values) const;
+	bool TryGetValues(const TCHAR* Key, TArray<FGuid>& Value) const;
+};
diff --git a/Engine/Source/Programs/SlateUGS/Private/UGSCore/DetectProjectSettingsTask.cpp b/Engine/Source/Programs/SlateUGS/Private/UGSCore/DetectProjectSettingsTask.cpp
new file mode 100644
--- /dev/null
+++ b/Engine/Source/Programs/SlateUGS/Private/UGSCore/DetectProjectSettingsTask.cpp
@@ -0,0 +1,274 @@
+// Copyright 1998-2017 Epic Games, Inc. All Rights Reserved.
+
+#include "DetectProjectSettingsTask.h"
+#include "OutputAdapters.h"
+#include "Utility.h"
+#include "HAL/PlatformProcess.h"
+#include "Misc/Paths.h"
+
+#define LOCTEXT_NAMESPACE "UnrealGameSync"
+
+FDetectProjectSettingsTask::FDetectProjectSettingsTask(TSharedRef<FPerforceConnection> InPerforce, const FString& InNewSelectedFileName, TSharedRef<FLineBasedTextWriter> InLog)
+	: NewSelectedFileName(InNewSelectedFileName)
+	, Perforce(MoveTemp(InPerforce))
+	, Log(MoveTemp(InLog))
+{
+}
+
+FDetectProjectSettingsTask::~FDetectProjectSettingsTask()
+{
+}
+
+TSharedRef<FModalTaskResult> FDetectProjectSettingsTask::Run(FEvent* AbortEvent)
+{
+	// TODO enable these if we want
+	//try
+	{
+		return RunInternal(AbortEvent);
+	}
+	//catch(FAbortException)
+	{
+		//return FModalTaskResult::Aborted();
+	}
+}
+
+TSharedRef<FModalTaskResult> FDetectProjectSettingsTask::RunInternal(FEvent* AbortEvent)
+{
+	// Get the P4PORT setting in this folder, so we can respect the contents of any P4CONFIG file
+	FString PrevDirectory = FPlatformProcess::GetCurrentWorkingDirectory();
+
+//	Directory.SetCurrentDirectory(Path.GetDirectoryName(NewSelectedFileName));
+//
+//	string ServerAndPort;
+//	if (Perforce.GetSetting("P4PORT", out ServerAndPort, Log))
+//	{
+//		Perforce = new PerforceConnection(null, null, ServerAndPort);
+//	}
+
+	// Get the Perforce server info
+	TSharedPtr<FPerforceInfoRecord> PerforceInfo;
+	if(!Perforce->Info(PerforceInfo, AbortEvent, Log.Get()))
+	{
+		return FModalTaskResult::Failure(LOCTEXT("NoPerforceInfo", "Couldn't get Perforce server info"));
+	}
+	if(PerforceInfo->UserName.Len() == 0)
+	{
+		return FModalTaskResult::Failure(LOCTEXT("MissingUserNameInP4Info", "Missing user name in call to p4 info"));
+	}
+	if(PerforceInfo->HostName.Len() == 0)
+	{
+		return FModalTaskResult::Failure(LOCTEXT("MissingHostNameInP4Info", "Missing host name in call to p4 info"));
+	}
+	ServerTimeZone = PerforceInfo->ServerTimeZone;
+
+	// Find all the clients on this machine
+	Log->Logf(TEXT("Enumerating clients on local machine..."));
+	TArray<FPerforceClientRecord> Clients;
+	if(!Perforce->FindClients(Clients, AbortEvent, Log.Get()))
+	{
+		return FModalTaskResult::Failure(LOCTEXT("NoClientsForHost", "Couldn't find any clients for this host."));
+	}
+
+	// Find any clients which are valid. If this is not exactly one, we should fail.
+	TArray<TSharedRef<FPerforceConnection>> CandidateClients;
+	for(const FPerforceClientRecord& Client : Clients)
+	{
+		// Make sure the client is well formed
+		if(Client.Name.Len() > 0 && (Client.Host.Len() > 0 || Client.Owner.Len() > 0) && Client.Root.Len() > 0)
+		{
+			// Require either a username or host name match
+			if((Client.Host.Len() == 0 || Client.Host == PerforceInfo->HostName) && (Client.Owner.Len() == 0 || Client.Owner == PerforceInfo->UserName))
+			{
+				if(!FUtility::IsFileUnderDirectory(*NewSelectedFileName, *Client.Root))
+				{
+					Log->Logf(TEXT("Rejecting %s due to root mismatch (%s)"), *Client.Name, *Client.Root);
+					continue;
+				}
+
+				TSharedRef<FPerforceConnection> CandidateClient = MakeShared<FPerforceConnection>(*PerforceInfo->UserName, *Client.Name, *Perforce->ServerAndPort);
+
+				bool bFileExists;
+				if(!CandidateClient->FileExists(NewSelectedFileName, bFileExists, AbortEvent, Log.Get()) || !bFileExists)
+				{
+					Log->Logf(TEXT("Rejecting %s due to file not existing in workspace"), *Client.Name);
+					continue;
+				}
+
+				TArray<FPerforceFileRecord> Records;
+				if(!CandidateClient->Stat(NewSelectedFileName, Records, AbortEvent, Log.Get()))
+				{
+					Log->Logf(TEXT("Rejecting %s due to %s not in depot"), *Client.Name, *NewSelectedFileName);
+					continue;
+				}
+
+				Records.RemoveAll([](const FPerforceFileRecord& Record){ return !Record.IsMapped; });
+				if(Records.Num() != 1)
+				{
+					Log->Logf(TEXT("Rejecting %s due to %d matching records"), *Client.Name, Records.Num());
+					continue;
+				}
+
+				Log->Logf(TEXT("Found valid client %s"), *Client.Name);
+				CandidateClients.Add(CandidateClient);
+			}
+		}
+	}
+
+	// Check there's only one client
+	if(CandidateClients.Num() == 0)
+	{
+		return FModalTaskResult::Failure(FText::Format(LOCTEXT("CouldNotFindWorkspace", "Couldn't find any Perforce workspace containing {0}."), FText::FromString(NewSelectedFileName)));
+	}
+	else if(CandidateClients.Num() > 1)
+	{
+		FString ClientList;
+		for(const TSharedRef<FPerforceConnection>& CandidateClient : CandidateClients)
+		{
+			ClientList += TEXT("\n");
+			ClientList += CandidateClient->ClientName;
+		}
+		return FModalTaskResult::Failure(FText::Format(LOCTEXT("FoundMultipleWorkspaces", "Found multiple workspaces containing {0}:\n{1}\n\nCannot determine which to use."), FText::FromString(FPaths::GetCleanFilename(NewSelectedFileName)), FText::FromString(ClientList)));
+	}
+
+	// Take the client we've chosen
+	PerforceClient = CandidateClients[0];
+
+	// Get the client path for the project file
+	if(!PerforceClient->ConvertToClientPath(NewSelectedFileName, NewSelectedClientFileName, AbortEvent, Log.Get()))
+	{
+		return FModalTaskResult::Failure(FText::Format(LOCTEXT("CouldNotGetClientPath", "Couldn't get client path for {0}"), FText::FromString(NewSelectedFileName)));
+	}
+
+	// Figure out where the engine is in relation to it
+	for(int EndIdx = NewSelectedClientFileName.Len() - 1;;EndIdx--)
+	{
+		if(EndIdx < 2)
+		{
+			return FModalTaskResult::Failure(FText::Format(LOCTEXT("CouldNotFindEngine", "Could not find engine in Perforce relative to project path ({0})"), FText::FromString(NewSelectedClientFileName)));
+		}
+		if(NewSelectedClientFileName[EndIdx] == '/')
+		{
+			bool bFileExists;
+			if(PerforceClient->FileExists(NewSelectedClientFileName.Mid(0, EndIdx) + "/Engine/Source/UnrealEditor.Target.cs", bFileExists, AbortEvent, Log.Get()) && bFileExists)
+			{
+				BranchClientPath = NewSelectedClientFileName.Mid(0, EndIdx);
+				break;
+			}
+		}
+	}
+	Log->Logf(TEXT("Found branch root at %s"), *BranchClientPath);
+
+	// Get the local branch root
+	FString BaseEditorTargetPath;
+	if(!PerforceClient->ConvertToLocalPath(BranchClientPath + TEXT("/Engine/Source/UnrealEditor.Target.cs"), BaseEditorTargetPath, AbortEvent, Log.Get()))
+	{
+		return FModalTaskResult::Failure(LOCTEXT("CouldNotGetEditorTarget", "Couldn't get local path for editor target file"));
+	}
+	BranchDirectoryName = FPaths::ConvertRelativePathToFull(FPaths::GetPath(FPaths::GetPath(FPaths::GetPath(BaseEditorTargetPath))));
+	FPaths::MakePlatformFilename(BranchDirectoryName);
+
+	// Find the editor target for this project
+	if(NewSelectedFileName.EndsWith(TEXT(".uproject")))
+	{
+		TArray<FPerforceFileRecord> Files;
+		if(PerforceClient->FindFiles(FPerforceUtils::GetClientOrDepotDirectoryName(*NewSelectedClientFileName) + TEXT("/Source/*Editor.Target.cs"), Files, AbortEvent, Log.Get()) && Files.Num() >= 1)
+		{
+			const FPerforceFileRecord* ValidFile = nullptr;
+			for(const FPerforceFileRecord& File : Files)
+			{
+				if(File.Action.Len() == 0 || !File.Action.Contains(TEXT("delete")))
+				{
+					int LastSlashIdx;
+					if(File.DepotPath.FindLastChar('/', LastSlashIdx))
+					{
+						// Grab the <Editor>.Target name
+						NewProjectEditorTarget = FPaths::GetBaseFilename(File.DepotPath.Mid(LastSlashIdx + 1));
+
+						// Remove the .Target
+						NewProjectEditorTarget = FPaths::GetBaseFilename(NewProjectEditorTarget);
+
+						Log->Logf(TEXT("Using %s as editor target name (from %s)"), *NewProjectEditorTarget, *File.DepotPath);
+						break;
+					}
+				}
+			}
+			if(NewProjectEditorTarget.Len() == 0)
+			{
+				Log->Logf(TEXT("Couldn't find any non-deleted editor targets for this project."));
+			}
+		}
+		else
+		{
+			Log->Logf(TEXT("Couldn't find any editor targets for this project."));
+		}
+	}
+
+	// Get a unique name for the project that's selected. For regular branches, this can be the depot path. For streams, we want to include the stream name to encode imports.
+	if(PerforceClient->GetActiveStream(StreamName, AbortEvent, Log.Get()))
+	{
+		FString ExpectedPrefix = FString::Printf(TEXT("//%s/"), *PerforceClient->ClientName);
+		if(!NewSelectedClientFileName.StartsWith(ExpectedPrefix))
+		{
+			return FModalTaskResult::Failure(FText::Format(LOCTEXT("UnexpectedClientPath", "Unexpected client path; expected '{0}' to begin with '{1}'"), FText::FromString(NewSelectedClientFileName), FText::FromString(ExpectedPrefix)));
+		}
+		FString StreamPrefix;
+		if(!TryGetStreamPrefix(PerforceClient.ToSharedRef(), StreamName, AbortEvent, Log.Get(), StreamPrefix))
+		{
+			return FModalTaskResult::Failure(FText::Format(LOCTEXT("FailedToGetStreamInfo", "Failed to get stream info for {0}"), FText::FromString(StreamName)));
+		}
+		NewSelectedProjectIdentifier = FString::Printf(TEXT("{0}/{1}"), *StreamPrefix, *NewSelectedClientFileName.Mid(ExpectedPrefix.Len()));
+	}
+	else
+	{
+		if(!PerforceClient->ConvertToDepotPath(NewSelectedClientFileName, NewSelectedProjectIdentifier, AbortEvent, Log.Get()))
+		{
+			return FModalTaskResult::Failure(FText::Format(LOCTEXT("CouldNotGetDepotPath", "Couldn't get depot path for {0}"), FText::FromString(NewSelectedFileName)));
+		}
+	}
+
+	// Read the project logo
+	if(NewSelectedFileName.EndsWith(TEXT(".uproject")))
+	{
+		/*
+		string LogoFileName = Path.Combine(Path.GetDirectoryName(NewSelectedFileName), "Build", "UnrealGameSync.png");
+		if(File.Exists(LogoFileName))
+		{
+			try
+			{
+				// Duplicate the image, otherwise we'll leave the file locked
+				using(Image Image = Image.FromFile(LogoFileName))
+				{
+					ProjectLogo = new Bitmap(Image);
+				}
+			}
+			catch
+			{
+				ProjectLogo = null;
+			}
+		}*/
+	}
+
+	// Succeeed!
+	return FModalTaskResult::Success();
+}
+
+bool FDetectProjectSettingsTask::TryGetStreamPrefix(TSharedRef<FPerforceConnection> Perforce, const FString& StreamName, FEvent* AbortEvent, FLineBasedTextWriter& Log, FString& OutStreamPrefix)
+{
+	FString CurrentStreamName = StreamName;
+	for(;;)
+	{
+		TSharedPtr<FPerforceSpec> StreamSpec;
+		if(!Perforce->TryGetStreamSpec(CurrentStreamName, StreamSpec, AbortEvent, Log))
+		{
+			return false;
+		}
+		if(StreamSpec->GetField(TEXT("Type")) != FString(TEXT("virtual")))
+		{
+			OutStreamPrefix = CurrentStreamName;
+			return true;
+		}
+		CurrentStreamName = StreamSpec->GetField(TEXT("Parent"));
+	}
+}
+
+#undef LOCTEXT_NAMESPACE
diff --git a/Engine/Source/Programs/SlateUGS/Private/UGSCore/DetectProjectSettingsTask.h b/Engine/Source/Programs/SlateUGS/Private/UGSCore/DetectProjectSettingsTask.h
new file mode 100644
--- /dev/null
+++ b/Engine/Source/Programs/SlateUGS/Private/UGSCore/DetectProjectSettingsTask.h
@@ -0,0 +1,36 @@
+// Copyright 1998-2017 Epic Games, Inc. All Rights Reserved.
+
+#pragma once
+
+#include "CoreMinimal.h"
+#include "ModalTask.h"
+#include "Perforce.h"
+#include "OutputAdapters.h"
+
+class FDetectProjectSettingsTask : public IModalTask
+{
+public:
+	const FString NewSelectedFileName;
+	TSharedRef<FPerforceConnection> Perforce;
+	TSharedPtr<FPerforceConnection> PerforceClient;
+	FString NewSelectedProjectIdentifier;
+	FString NewProjectEditorTarget;
+	FString BranchClientPath;
+	FString BranchDirectoryName;
+	FString NewSelectedClientFileName;
+	FString StreamName;
+//	public Image ProjectLogo;
+	FTimespan ServerTimeZone;
+//	TextWriter Log;
+
+	FDetectProjectSettingsTask(TSharedRef<FPerforceConnection> InPerforce, const FString& InNewSelectedFileName, TSharedRef<FLineBasedTextWriter> InLog);
+	~FDetectProjectSettingsTask();
+
+	virtual TSharedRef<FModalTaskResult> Run(FEvent* AbortEvent) override;
+
+private:
+	TSharedRef<FLineBasedTextWriter> Log;
+
+	TSharedRef<FModalTaskResult> RunInternal(FEvent* AbortEvent);
+	static bool TryGetStreamPrefix(TSharedRef<FPerforceConnection> Perforce, const FString& StreamName, FEvent* AbortEvent, FLineBasedTextWriter& Log, FString& OutStreamPrefix);
+};
diff --git a/Engine/Source/Programs/SlateUGS/Private/UGSCore/FileFilter.cpp b/Engine/Source/Programs/SlateUGS/Private/UGSCore/FileFilter.cpp
new file mode 100644
--- /dev/null
+++ b/Engine/Source/Programs/SlateUGS/Private/UGSCore/FileFilter.cpp
@@ -0,0 +1,396 @@
+// Copyright 1998-2017 Epic Games, Inc. All Rights Reserved.
+
+#include "FileFilter.h"
+
+/// <summary>
+/// Node within a filter tree. Each node matches a single path fragment - a folder or file name, with an include or exclude rule consisting of a sequence of nodes. 
+/// </summary>
+struct FFileFilter::FNode
+{
+	/// <summary>
+	/// Node which this is parented to. Null for the root node.
+	/// </summary>
+	FNode* const Parent;
+
+	/// <summary>
+	/// Pattern to match for this node. May contain * or ? wildcards.
+	/// </summary>
+	const FString Pattern;
+
+	/// <summary>
+	/// Highest include rule number matched by this node or any child nodes.
+	/// </summary>
+	int MaxIncludeRuleNumber = -1;
+
+	/// <summary>
+	/// Highest exclude rule number matched by this node or any child nodes.
+	/// </summary>
+	int MaxExcludeRuleNumber = -1;
+
+	/// <summary>
+	/// Child branches of this node, distinct by the the pattern for each.
+	/// </summary>
+	TArray<FNode*> Branches;
+
+	/// <summary>
+	/// If a path matches the sequence of nodes ending in this node, the number of the rule which added it. -1 if this was not the last node in a rule.
+	/// </summary>
+	int RuleNumber;
+
+	/// <summary>
+	/// Whether a rule terminating in this node should include files or exclude files.
+	/// </summary>
+	EFileFilterType Type;
+
+	/// <summary>
+	/// Default constructor.
+	/// </summary>
+	FNode(FNode* InParent, const FString& InPattern)
+		: Parent(InParent)
+		, Pattern(InPattern)
+	{
+		RuleNumber = -1;
+		Type = EFileFilterType::Exclude;
+	}
+
+	/// <summary>
+	/// Destructor.
+	/// </summary>
+	~FNode()
+	{
+		for(FNode* Branch : Branches)
+		{
+			delete Branch;
+			Branch = nullptr;
+		}
+	}
+
+	/// <summary>
+	/// Determine if the given token matches the pattern in this node
+	/// </summary>
+	bool IsMatch(const FString& Token) const
+	{
+		if(Pattern.EndsWith("."))
+		{
+			return !Token.Contains(TEXT(".")) && IsMatch(Token, 0, Pattern.Mid(0, Pattern.Len() - 1), 0);
+		}
+		else
+		{
+			return IsMatch(Token, 0, Pattern, 0);
+		}
+	}
+
+	/// <summary>
+	/// Determine if a given token matches a pattern, starting from the given positions within each string.
+	/// </summary>
+	static bool IsMatch(const FString& Token, int TokenIdx, const FString& Pattern, int PatternIdx)
+	{
+		for (; ; )
+		{
+			if (PatternIdx == Pattern.Len())
+			{
+				return (TokenIdx == Token.Len());
+			}
+			else if (Pattern[PatternIdx] == '*')
+			{
+				for (int NextTokenIdx = Token.Len(); NextTokenIdx >= TokenIdx; NextTokenIdx--)
+				{
+					if (IsMatch(Token, NextTokenIdx, Pattern, PatternIdx + 1))
+					{
+						return true;
+					}
+				}
+				return false;
+			}
+			else if (TokenIdx < Token.Len() && (FChar::ToLower(Token[TokenIdx]) == FChar::ToLower(Pattern[PatternIdx]) || Pattern[PatternIdx] == '?'))
+			{
+				TokenIdx++;
+				PatternIdx++;
+			}
+			else
+			{
+				return false;
+			}
+		}
+	}
+
+	/// <summary>
+	/// Converts this node back to a pattern from the root
+	/// </summary>
+	/// <returns>Path to this node</returns>
+	FString ToString() const
+	{
+		if(Parent == nullptr)
+		{
+			return TEXT("/");
+		}
+		else if(Parent->Parent == nullptr)
+		{
+			return FString(TEXT("/")) + Pattern;
+		}
+		else
+		{
+			return Parent->ToString() + FString(TEXT("/")) + Pattern;
+		}
+	}
+};
+
+FFileFilter::FFileFilter(EFileFilterType DefaultType)
+{
+	RootNode = new FNode(nullptr, "");
+
+	DefaultNode = new FNode(RootNode, "...");
+	DefaultNode->Type = DefaultType;
+}
+
+FFileFilter::FFileFilter(const FFileFilter& Other)
+	: FFileFilter(Other.DefaultNode->Type)
+{
+	for(const FNode* OtherRule : Other.Rules)
+	{
+		AddRule(OtherRule->ToString(), OtherRule->Type);
+	}
+}
+
+FFileFilter::~FFileFilter()
+{
+	delete RootNode;
+}
+
+/*
+void FFileFilter::ReadRulesFromFile(string FileName, string SectionName, params string[] AllowTags)
+{
+	bool bInSection = false;
+	foreach(string Line in File.ReadAllLines(FileName))
+	{
+		string TrimLine = Line.Trim();
+		if(!TrimLine.StartsWith(";") && TrimLine.Length > 0)
+		{
+			if(TrimLine.StartsWith("["))
+			{
+				bInSection = (TrimLine == "[" + SectionName + "]");
+			}
+			else if(bInSection)
+			{
+				AddRule(Line, AllowTags);
+			}
+		}
+	}
+}
+*/
+
+void FFileFilter::AddRule(const FString& Pattern)
+{
+	if(Pattern.StartsWith(TEXT("-")))
+	{
+		AddRule(Pattern.Mid(1), EFileFilterType::Exclude);
+	}
+	else
+	{
+		AddRule(Pattern, EFileFilterType::Include);
+	}
+}
+
+void FFileFilter::AddRule(const FString& Pattern, EFileFilterType Type)
+{
+	FString NormalizedPattern = Pattern;
+	NormalizedPattern.ReplaceInline(TEXT("\\"), TEXT("/"));
+
+	// We don't want a slash at the start, but if there was not one specified, it's not anchored to the root of the tree.
+	if (NormalizedPattern.StartsWith(TEXT("/")))
+	{
+		NormalizedPattern = NormalizedPattern.Mid(1);
+	}
+	else if(!NormalizedPattern.StartsWith("..."))
+	{
+		NormalizedPattern = TEXT(".../") + NormalizedPattern;
+	}
+
+	// All directories indicate a wildcard match
+	if (NormalizedPattern.EndsWith(TEXT("/")))
+	{
+		NormalizedPattern += TEXT("...");
+	}
+
+	// Replace any directory wildcards mid-string
+	for(int Idx = 0;;)
+	{
+		Idx = NormalizedPattern.Find(TEXT("..."), ESearchCase::CaseSensitive, ESearchDir::FromStart, Idx);
+		if(Idx == INDEX_NONE)
+		{
+			break;
+		}
+
+		if (Idx > 0 && NormalizedPattern[Idx - 1] != '/')
+		{
+			NormalizedPattern.InsertAt(Idx, TEXT("*/"));
+			Idx++;
+		}
+
+		Idx += 3;
+
+		if (Idx < NormalizedPattern.Len() && NormalizedPattern[Idx] != '/')
+		{
+			NormalizedPattern.InsertAt(Idx, TEXT("/*"));
+			Idx += 2;
+		}
+	}
+
+	// Split the pattern into fragments
+	TArray<FString> BranchPatterns;
+	NormalizedPattern.ParseIntoArray(BranchPatterns, TEXT("/"), false);
+
+	// Add it into the tree
+	FNode* LastNode = RootNode;
+	for(const FString& BranchPattern : BranchPatterns)
+	{
+		FNode* NextNode = nullptr;
+		for(FNode* BranchNode : LastNode->Branches)
+		{
+			if(BranchNode->Pattern == BranchPattern)
+			{
+				NextNode = BranchNode;
+				break;
+			}
+		}
+		if (NextNode == nullptr)
+		{
+			NextNode = new FNode(LastNode, BranchPattern);
+			LastNode->Branches.Add(NextNode);
+		}
+		LastNode = NextNode;
+	}
+
+	// We've reached the end of the pattern, so mark it as a leaf node
+	Rules.Add(LastNode);
+	LastNode->RuleNumber = Rules.Num() - 1;
+	LastNode->Type = Type;
+
+	// Update the maximums along that path
+	for (FNode* UpdateNode = LastNode; UpdateNode != nullptr; UpdateNode = UpdateNode->Parent)
+	{
+		if (Type == EFileFilterType::Include)
+		{
+			UpdateNode->MaxIncludeRuleNumber = LastNode->RuleNumber;
+		}
+		else
+		{
+			UpdateNode->MaxExcludeRuleNumber = LastNode->RuleNumber;
+		}
+	}
+}
+
+TArray<FString> TokenizePath(const FString& Path)
+{
+	int Idx = 0;
+	while(Idx < Path.Len() && (Path[Idx] == '\\' || Path[Idx] == '/'))
+	{
+		Idx++;
+	}
+
+	TArray<FString> Tokens;
+
+	int PrevIdx = Idx;
+	for(; Idx < Path.Len(); Idx++)
+	{
+		if(Path[Idx] == '/' || Path[Idx] == '\\')
+		{
+			Tokens.Add(Path.Mid(PrevIdx, Idx - PrevIdx));
+			PrevIdx = Idx + 1;
+		}
+	}
+	if(Idx > PrevIdx)
+	{
+		Tokens.Add(Path.Mid(PrevIdx, Idx - PrevIdx));
+	}
+
+	return Tokens;
+}
+
+bool FFileFilter::Matches(const FString& FileName) const
+{
+	TArray<FString> Tokens = TokenizePath(FileName);
+
+	const FNode* MatchingNode = FindMatchingNode(RootNode, Tokens, 0, DefaultNode);
+
+	return MatchingNode->Type == EFileFilterType::Include;
+}
+
+bool FFileFilter::PossiblyMatches(const FString& FolderName) const
+{
+	TArray<FString> Tokens = TokenizePath(FolderName);
+	Tokens.Add(TEXT(""));
+
+	const FNode* MatchingNode = FindMatchingNode(RootNode, Tokens, 0, DefaultNode);
+
+	return MatchingNode->Type == EFileFilterType::Include || HighestPossibleIncludeMatch(RootNode, Tokens, 0, MatchingNode->RuleNumber) > MatchingNode->RuleNumber;
+}
+
+const FFileFilter::FNode* FFileFilter::FindMatchingNode(const FNode* CurrentNode, const TArray<FString>& Tokens, int TokenIdx, const FNode* CurrentBestNode)
+{
+	// If we've matched all the input tokens, check if this rule is better than any other we've seen
+	if (TokenIdx == Tokens.Num())
+	{
+		return (CurrentNode->RuleNumber > CurrentBestNode->RuleNumber) ? CurrentNode : CurrentBestNode;
+	}
+
+	// If there is no rule under the current node which is better than the current best node, early out
+	if (CurrentNode->MaxIncludeRuleNumber <= CurrentBestNode->RuleNumber && CurrentNode->MaxExcludeRuleNumber <= CurrentBestNode->RuleNumber)
+	{
+		return CurrentBestNode;
+	}
+
+	// Test all the branches for one that matches
+	const FNode* BestNode = CurrentBestNode;
+	for(const FNode* Branch : CurrentNode->Branches)
+	{
+		if (Branch->Pattern == "...")
+		{
+			for (int NextTokenIdx = Tokens.Num(); NextTokenIdx >= TokenIdx; NextTokenIdx--)
+			{
+				BestNode = FindMatchingNode(Branch, Tokens, NextTokenIdx, BestNode);
+			}
+		}
+		else
+		{
+			if (Branch->IsMatch(Tokens[TokenIdx]))
+			{
+				BestNode = FindMatchingNode(Branch, Tokens, TokenIdx + 1, BestNode);
+			}
+		}
+	}
+	return BestNode;
+}
+
+int FFileFilter::HighestPossibleIncludeMatch(const FNode* CurrentNode, const TArray<FString>& Tokens, int TokenIdx, int CurrentBestRuleNumber)
+{
+	// If we've matched all the input tokens, check if this rule is better than any other we've seen
+	if (TokenIdx == Tokens.Num())
+	{
+		return FMath::Max(CurrentBestRuleNumber, CurrentNode->MaxIncludeRuleNumber);
+	}
+
+	// Test all the branches for one that matches
+	int BestRuleNumber = CurrentBestRuleNumber;
+	if (CurrentNode->MaxIncludeRuleNumber > BestRuleNumber)
+	{
+		for(const FNode* Branch : CurrentNode->Branches)
+		{
+			if (Branch->Pattern == "...")
+			{
+				if (Branch->MaxIncludeRuleNumber > BestRuleNumber)
+				{
+					BestRuleNumber = Branch->MaxIncludeRuleNumber;
+				}
+			}
+			else
+			{
+				if (Branch->IsMatch(Tokens[TokenIdx]))
+				{
+					BestRuleNumber = HighestPossibleIncludeMatch(Branch, Tokens, TokenIdx + 1, BestRuleNumber);
+				}
+			}
+		}
+	}
+	return BestRuleNumber;
+}
diff --git a/Engine/Source/Programs/SlateUGS/Private/UGSCore/FileFilter.h b/Engine/Source/Programs/SlateUGS/Private/UGSCore/FileFilter.h
new file mode 100644
--- /dev/null
+++ b/Engine/Source/Programs/SlateUGS/Private/UGSCore/FileFilter.h
@@ -0,0 +1,121 @@
+// Copyright 1998-2017 Epic Games, Inc. All Rights Reserved.
+
+#pragma once
+
+#include "CoreMinimal.h"
+
+struct FFileFilterNode;
+
+/// <summary>
+/// Indicates whether files which match a pattern should be included or excluded
+/// </summary>
+enum class EFileFilterType
+{
+	Include,
+	Exclude,
+};
+
+/// <summary>
+/// Stores a set of rules, similar to p4 path specifications, which can be used to efficiently filter files in or out of a set.
+/// The rules are merged into a tree with nodes for each path fragments and rule numbers in each leaf, allowing tree traversal in an arbitrary order 
+/// to how they are applied.
+/// </summary>
+class FFileFilter
+{
+public:
+	/// <summary>
+	/// Default constructor
+	/// </summary>
+	FFileFilter(EFileFilterType DefaultType = EFileFilterType::Exclude);
+
+	/// <summary>
+	/// Copy constructor
+	/// </summary>
+	/// <param name="Other">Filter to copy from</param>
+	FFileFilter(const FFileFilter& Other);
+
+	/// <summary>
+	/// Constructs a file filter from a p4-style filespec. Exclude lines are prefixed with a - character.
+	/// </summary>
+	/// <param name="Lines">Lines to parse rules from</param>
+	FFileFilter(const TArray<FString>& Lines);
+
+	/// <summary>
+	/// Destructor
+	/// </summary>
+	~FFileFilter();
+
+	/// <summary>
+	/// Reads a configuration file split into sections
+	/// </summary>
+	/// <param name="Filter"></param>
+	/// <param name="RulesFileName"></param>
+	/// <param name="Conditions"></param>
+	void ReadRulesFromFile(const FString& FileName, const FString& SectionName);
+
+	/// <summary>
+	/// Adds an include or exclude rule to the filter, treating anything with a '-' prefix as an exclude rule, and anything else as an include rule.
+	/// </summary>
+	/// <param name="Pattern">The pattern which the rule should match</param>
+	/// <param name="bInclude">Whether to include or exclude files matching this rule</param>
+	void AddRule(const FString& Pattern);
+
+	/// <summary>
+	/// Adds an include or exclude rule to the filter
+	/// </summary>
+	/// <param name="Pattern">The pattern which the rule should match</param>
+	/// <param name="bInclude">Whether to include or exclude files matching this rule</param>
+	void AddRule(const FString& Pattern, EFileFilterType Type);
+
+	/// <summary>
+	/// Determines whether the given file matches the filter
+	/// </summary>
+	/// <param name="FileName">File to match</param>
+	/// <returns>True if the file passes the filter</returns>
+	bool Matches(const FString& FileName) const;
+
+	/// <summary>
+	/// Determines whether it's possible for anything within the given folder name to match the filter. Useful to early out of recursive file searches.
+	/// </summary>
+	/// <param name="FolderName">File to match</param>
+	/// <returns>True if the file passes the filter</returns>
+	bool PossiblyMatches(const FString& FolderName) const;
+
+private:
+	struct FNode;
+
+	/// <summary>
+	/// Root node for the tree.
+	/// </summary>
+	FNode* RootNode;
+
+	/// <summary>
+	/// The default node, which will match any path
+	/// </summary>
+	FNode* DefaultNode;
+
+	/// <summary>
+	/// Terminating nodes for each rule added to the filter
+	/// </summary>
+	TArray<FNode*> Rules;
+
+	/// <summary>
+	/// Finds the node which matches a given list of tokens.
+	/// </summary>
+	/// <param name="CurrentNode"></param>
+	/// <param name="Tokens"></param>
+	/// <param name="TokenIdx"></param>
+	/// <param name="CurrentBestNode"></param>
+	/// <returns></returns>
+	static const FNode* FindMatchingNode(const FNode* CurrentNode, const TArray<FString>& Tokens, int TokenIdx, const FNode* CurrentBestNode);
+
+	/// <summary>
+	/// Returns the highest possible rule number which can match the given list of input tokens, assuming that the list of input tokens is incomplete.
+	/// </summary>
+	/// <param name="CurrentNode">The current node being checked</param>
+	/// <param name="Tokens">The tokens to match</param>
+	/// <param name="TokenIdx">Current token index</param>
+	/// <param name="CurrentBestRuleNumber">The highest rule number seen so far. Used to optimize tree traversals.</param>
+	/// <returns>New highest rule number</returns>
+	static int HighestPossibleIncludeMatch(const FNode* CurrentNode, const TArray<FString>& Tokens, int TokenIdx, int CurrentBestRuleNumber);
+};
diff --git a/Engine/Source/Programs/SlateUGS/Private/UGSCore/ModalTask.cpp b/Engine/Source/Programs/SlateUGS/Private/UGSCore/ModalTask.cpp
new file mode 100644
--- /dev/null
+++ b/Engine/Source/Programs/SlateUGS/Private/UGSCore/ModalTask.cpp
@@ -0,0 +1,51 @@
+// Copyright 1998-2017 Epic Games, Inc. All Rights Reserved.
+
+#include "ModalTask.h"
+
+//// FModalTaskResult ////
+
+FModalTaskResult::FModalTaskResult(bool bInSucceeded, const FText& InMessage)
+	: bSucceeded(bInSucceeded)
+	, Message(InMessage)
+{
+}
+
+FModalTaskResult::~FModalTaskResult()
+{
+}
+
+bool FModalTaskResult::Succeeded() const
+{
+	return bSucceeded;
+}
+
+bool FModalTaskResult::Failed() const
+{
+	return !bSucceeded;
+}
+
+const FText& FModalTaskResult::GetMessage() const
+{
+	return Message;
+}
+
+TSharedRef<FModalTaskResult> FModalTaskResult::Success()
+{
+	return MakeShared<FModalTaskResult>(true, FText::GetEmpty());
+}
+
+TSharedRef<FModalTaskResult> FModalTaskResult::Failure(const FText& Message)
+{
+	return MakeShared<FModalTaskResult>(false, Message);
+}
+
+TSharedRef<FModalTaskResult> FModalTaskResult::Aborted()
+{
+	return MakeShared<FModalTaskResult>(false, FText::GetEmpty());
+}
+
+//// IModalTask ////
+
+IModalTask::~IModalTask()
+{
+}
diff --git a/Engine/Source/Programs/SlateUGS/Private/UGSCore/ModalTask.h b/Engine/Source/Programs/SlateUGS/Private/UGSCore/ModalTask.h
new file mode 100644
--- /dev/null
+++ b/Engine/Source/Programs/SlateUGS/Private/UGSCore/ModalTask.h
@@ -0,0 +1,32 @@
+// Copyright 1998-2017 Epic Games, Inc. All Rights Reserved.
+
+#pragma once
+
+#include "CoreMinimal.h"
+#include "Templates/SharedPointer.h"
+
+class FModalTaskResult
+{
+public:
+	FModalTaskResult(bool bSucceeded, const FText& Message);
+	~FModalTaskResult();
+
+	bool Succeeded() const;
+	bool Failed() const;
+	const FText& GetMessage() const;
+
+	static TSharedRef<FModalTaskResult> Success();
+	static TSharedRef<FModalTaskResult> Failure(const FText& Message);
+	static TSharedRef<FModalTaskResult> Aborted();
+
+private:
+	const bool bSucceeded;
+	const FText Message;
+};
+
+class IModalTask
+{
+public:
+	virtual ~IModalTask();
+	virtual TSharedRef<FModalTaskResult> Run(FEvent* AbortEvent) = 0;
+};
diff --git a/Engine/Source/Programs/SlateUGS/Private/UGSCore/OutputAdapters.cpp b/Engine/Source/Programs/SlateUGS/Private/UGSCore/OutputAdapters.cpp
new file mode 100644
--- /dev/null
+++ b/Engine/Source/Programs/SlateUGS/Private/UGSCore/OutputAdapters.cpp
@@ -0,0 +1,435 @@
+// Copyright 1998-2017 Epic Games, Inc. All Rights Reserved.
+
+#include "OutputAdapters.h"
+#include "Utility.h"
+
+//// FLineBasedTextWriter ////
+
+FLineBasedTextWriter::FLineBasedTextWriter()
+{
+}
+
+FLineBasedTextWriter::~FLineBasedTextWriter()
+{
+}
+
+void FLineBasedTextWriter::Serialize(const TCHAR* V, ELogVerbosity::Type Verbosity, const class FName& Category)
+{
+	WriteLine(V);
+}
+
+void FLineBasedTextWriter::Write(TCHAR Character)
+{
+	if(Character == '\n')
+	{
+		FlushLine(CurrentLine);
+		CurrentLine.Empty();
+	}
+	else if(Character != '\r')
+	{
+		CurrentLine.AppendChar(Character);
+	}
+}
+
+void FLineBasedTextWriter::WriteLine(const FString& Line)
+{
+	TArray<FString> SubLines;
+	Line.ParseIntoArray(SubLines, TEXT("\n"), false);
+
+	for(FString& SubLine: SubLines)
+	{
+		FlushLine(SubLine.TrimEnd());
+	}
+}
+
+//// FNullTextWriter ////
+
+void FNullTextWriter::FlushLine(const FString& Line)
+{
+}
+
+//// FBufferedTextWriter ////
+
+FBufferedTextWriter::FBufferedTextWriter()
+{
+}
+
+FBufferedTextWriter::~FBufferedTextWriter()
+{
+}
+
+void FBufferedTextWriter::Attach(const TSharedRef<FLineBasedTextWriter>& InInner)
+{
+	Detach();
+
+	Inner = InInner;
+
+	for(const FString& BufferedLine : BufferedLines)
+	{
+		Inner->FlushLine(BufferedLine);
+	}
+	BufferedLines.Empty();
+}
+
+void FBufferedTextWriter::Detach()
+{
+	Inner.Reset();
+}
+
+void FBufferedTextWriter::FlushLine(const FString& Line)
+{
+	if(Inner.IsValid())
+	{
+		Inner->FlushLine(Line);
+	}
+	else
+	{
+		BufferedLines.Add(Line);
+	}
+}
+
+//// FPrefixedTextWriter ////
+
+FPrefixedTextWriter::FPrefixedTextWriter(const TCHAR* InPrefix, TSharedRef<FLineBasedTextWriter> InInner)
+	: Prefix(InPrefix)
+	, Inner(InInner)
+{
+}
+
+FPrefixedTextWriter::~FPrefixedTextWriter()
+{
+}
+
+void FPrefixedTextWriter::FlushLine(const FString& Line)
+{
+	Inner->WriteLine(Prefix + Line);
+}
+
+//// FBoundedLogWriter ////
+
+FBoundedLogWriter::FBoundedLogWriter(const TCHAR* InFileName, int32 InMaxSize)
+	: FileName(InFileName)
+	, BackupFileName(FileName + TEXT(".bak"))
+	, MaxSize(InMaxSize)
+	, Inner(nullptr)
+{
+	OpenInner();
+}
+
+FBoundedLogWriter::~FBoundedLogWriter()
+{
+	CloseInner();
+}
+
+void FBoundedLogWriter::FlushLine(const FString& Line)
+{
+	/*
+			if(Inner != null)
+			{
+				Inner.WriteLine(Line);
+				if(Inner.BaseStream.Position > MaxSize)
+				{
+					CloseInner();
+					OpenInner();
+				}
+			}
+	*/
+}
+
+void FBoundedLogWriter::OpenInner()
+{
+/*	CloseInner();
+	try
+	{
+		if(File.Exists(BackupFileName))
+		{
+			File.Delete(BackupFileName);
+		}
+		if(File.Exists(FileName))
+		{
+			File.Move(FileName, BackupFileName);
+		}
+
+		Inner = new StreamWriter(FileName);
+		Inner.AutoFlush = true;
+	}
+	catch(Exception)
+	{
+		Inner = null;
+	}*/
+}
+
+void FBoundedLogWriter::CloseInner()
+{
+	/*
+	if(Inner != null)
+	{
+		Inner.Close();
+		Inner = null;
+	}*/
+}
+
+//// FComposedTextWriter ////
+
+FComposedTextWriter::FComposedTextWriter(const TArray<TSharedRef<FLineBasedTextWriter>>& InInners)
+	: Inners(InInners)
+{
+}
+
+FComposedTextWriter::~FComposedTextWriter()
+{
+}
+
+void FComposedTextWriter::FlushLine(const FString& Line)
+{
+	for(TSharedRef<FLineBasedTextWriter>& Inner : Inners)
+	{
+		Inner->FlushLine(Line);
+	}
+}
+
+//// FProgressValue ////
+
+FProgressValue::FProgressValue()
+{
+	Clear();
+}
+
+FProgressValue::~FProgressValue()
+{
+}
+
+void FProgressValue::Clear()
+{
+	State = TTuple<FString,float>(TEXT("Starting..."), 0.0f);
+
+	Ranges.Empty();
+	Ranges.Add(TTuple<float, float>(0.0f, 1.0f));
+}
+
+TTuple<FString, float> FProgressValue::GetCurrent() const
+{
+	return State;
+}
+
+void FProgressValue::Set(const TCHAR* Message)
+{
+	if(Ranges.Num() == 1)
+	{
+		State = TTuple<FString,float>(Message, State.Value);
+	}
+}
+
+void FProgressValue::Set(const TCHAR* Message, float Fraction)
+{
+	if(Ranges.Num() == 1)
+	{
+		State = TTuple<FString, float>(Message, RelativeToAbsoluteFraction(Fraction));
+	}
+	else
+	{
+		State = TTuple<FString, float>(State.Key, RelativeToAbsoluteFraction(Fraction));
+	}
+}
+
+void FProgressValue::Set(float Fraction)
+{
+	State = TTuple<FString, float>(State.Key, RelativeToAbsoluteFraction(Fraction));
+}
+
+void FProgressValue::Increment(float Fraction)
+{
+	Set(State.Value + RelativeToAbsoluteFraction(Fraction));
+}
+
+void FProgressValue::Push(float MaxFraction)
+{
+	Ranges.Push(TTuple<float,float>(State.Value, RelativeToAbsoluteFraction(MaxFraction)));
+}
+
+void FProgressValue::Pop()
+{
+	if(Ranges.Num() > 1)
+	{
+		State = TTuple<FString,float>(State.Key, Ranges.Pop().Value);
+	}
+}
+
+float FProgressValue::RelativeToAbsoluteFraction(float Fraction)
+{
+	TTuple<float, float> Range = Ranges.Top();
+	return Range.Key + (Range.Value - Range.Key) * Fraction;
+}
+
+//// FProgressTextWriter ////
+
+const FString FProgressTextWriter::DirectivePrefix = TEXT("@progress ");
+
+FProgressTextWriter::FProgressTextWriter(const FProgressValue& InValue, const TSharedRef<FLineBasedTextWriter>& InInner)
+	: Value(InValue)
+	, Inner(InInner)
+{ 
+}
+
+FProgressTextWriter::~FProgressTextWriter()
+{
+}
+
+void FProgressTextWriter::FlushLine(const FString& Line)
+{
+	if(Line.StartsWith(DirectivePrefix))
+	{
+		// Line that just contains a progress directive
+		bool bSkipLine = false;
+		ProcessInternal(Line.Mid(DirectivePrefix.Len()), bSkipLine);
+	}
+	else
+	{
+		bool bSkipLine = false;
+		FString RemainingLine = Line;
+
+		// Look for a progress directive at the end of a line, in square brackets
+		FString TrimLine = Line.TrimEnd();
+		if(TrimLine.EndsWith(TEXT("]")))
+		{
+			for(int LastIdx = TrimLine.Len() - 2; LastIdx >= 0 && TrimLine[LastIdx] != ']'; LastIdx--)
+			{
+				if(TrimLine[LastIdx] == '[')
+				{
+					FString DirectiveSubstring = TrimLine.Mid(LastIdx + 1, TrimLine.Len() - LastIdx - 2);
+					if(DirectiveSubstring.StartsWith(DirectivePrefix))
+					{
+						ProcessInternal(DirectiveSubstring.Mid(DirectivePrefix.Len()), bSkipLine);
+						RemainingLine = Line.Mid(0, LastIdx).TrimEnd();
+					}
+					break;
+				}
+			}
+		}
+
+		if(!bSkipLine)
+		{
+			Inner->WriteLine(RemainingLine);
+		}
+	}
+}
+
+void FProgressTextWriter::ProcessInternal(const FString& Line, bool& bSkipLine)
+{
+	TArray<FString> Tokens = ParseTokens(Line);
+	for(int TokenIdx = 0; TokenIdx < Tokens.Num(); )
+	{
+		float Fraction;
+		if(ReadFraction(Tokens, TokenIdx, Fraction))
+		{
+			Value.Set(Fraction);
+		}
+		else if(Tokens[TokenIdx] == TEXT("push"))
+		{
+			TokenIdx++;
+			if(ReadFraction(Tokens, TokenIdx, Fraction))
+			{
+				Value.Push(Fraction);
+			}
+		}
+		else if(Tokens[TokenIdx] == TEXT("pop"))
+		{
+			TokenIdx++;
+			Value.Pop();
+		}
+		else if(Tokens[TokenIdx] == TEXT("increment"))
+		{
+			TokenIdx++;
+			if(ReadFraction(Tokens, TokenIdx, Fraction))
+			{
+				Value.Increment(Fraction);
+			}
+		}
+		else if(Tokens[TokenIdx] == TEXT("skipline"))
+		{
+			TokenIdx++;
+			bSkipLine = true;
+		}
+		else if(Tokens[TokenIdx].Len() >= 2 && ((Tokens[TokenIdx].StartsWith(TEXT("\'")) && Tokens[TokenIdx].EndsWith(TEXT("\'"))) || (Tokens[TokenIdx].StartsWith(TEXT("\"")) && Tokens[TokenIdx].EndsWith(TEXT("\"")))))
+		{
+			const FString& Message = Tokens[TokenIdx++];
+			Value.Set(*Message.Mid(1, Message.Len() - 2));
+		}
+		else
+		{
+			TokenIdx++;
+		}
+	}
+}
+
+TArray<FString> FProgressTextWriter::ParseTokens(const FString& Line)
+{
+	TArray<FString> Tokens;
+	for(int Idx = 0;;)
+	{
+		// Skip whitespace
+		while(Idx < Line.Len() && FChar::IsWhitespace(Line[Idx]))
+		{
+			Idx++;
+		}
+		if(Idx == Line.Len())
+		{
+			break;
+		}
+
+		// Read the next token
+		if(FChar::IsAlnum(Line[Idx]))
+		{
+			int StartIdx = Idx++;
+			while(Idx < Line.Len() && FChar::IsAlnum(Line[Idx]))
+			{
+				Idx++;
+			}
+			Tokens.Add(Line.Mid(StartIdx, Idx - StartIdx));
+		}
+		else if(Line[Idx] == '\'' || Line[Idx] == '\"')
+		{
+			int StartIdx = Idx++;
+			while(Idx < Line.Len() && Line[Idx] != Line[StartIdx])
+			{
+				Idx++;
+			}
+			Tokens.Add(Line.Mid(StartIdx, ++Idx - StartIdx));
+		}
+		else
+		{
+			Tokens.Add(Line.Mid(Idx++, 1));
+		}
+	}
+	return Tokens;
+}
+
+bool FProgressTextWriter::ReadFraction(const TArray<FString>& Tokens, int& TokenIdx, float& Fraction)
+{
+	// Read a fraction in the form x%
+	if(TokenIdx + 2 <= Tokens.Num() && Tokens[TokenIdx + 1] == TEXT("%"))
+	{
+		int Numerator;
+		if(FUtility::TryParse(*Tokens[TokenIdx], Numerator))
+		{
+			Fraction = (float)Numerator / 100.0f;
+			TokenIdx += 2;
+			return true;
+		}
+	}
+			
+	// Read a fraction in the form x/y
+	if(TokenIdx + 3 <= Tokens.Num() && Tokens[TokenIdx + 1] == "/")
+	{
+		int Numerator, Denominator;
+		if(FUtility::TryParse(*Tokens[TokenIdx], Numerator) && FUtility::TryParse(*Tokens[TokenIdx + 2], Denominator))
+		{
+			Fraction = (float)Numerator / (float)Denominator;
+			TokenIdx += 3;
+			return true;
+		}
+	}
+
+	Fraction = 0.0f;
+	return false;
+}
diff --git a/Engine/Source/Programs/SlateUGS/Private/UGSCore/OutputAdapters.h b/Engine/Source/Programs/SlateUGS/Private/UGSCore/OutputAdapters.h
new file mode 100644
--- /dev/null
+++ b/Engine/Source/Programs/SlateUGS/Private/UGSCore/OutputAdapters.h
@@ -0,0 +1,122 @@
+// Copyright 1998-2017 Epic Games, Inc. All Rights Reserved.
+
+#pragma once
+
+#include "CoreMinimal.h"
+
+class FLineBasedTextWriter : public FOutputDevice
+{
+public:
+	FLineBasedTextWriter();
+	virtual ~FLineBasedTextWriter();
+	void Write(TCHAR Character);
+	void WriteLine(const FString& Line);
+	virtual void Serialize(const TCHAR* V, ELogVerbosity::Type Verbosity, const class FName& Category) override final;
+	virtual void FlushLine(const FString& Line) = 0;
+
+private:
+	FString CurrentLine;
+};
+
+class FNullTextWriter : public FLineBasedTextWriter
+{
+public:
+	virtual void FlushLine(const FString& Line) override;
+};
+
+class FBufferedTextWriter : public FLineBasedTextWriter
+{
+public:
+	FBufferedTextWriter();
+	virtual ~FBufferedTextWriter();
+
+	void Attach(const TSharedRef<FLineBasedTextWriter>& InInner);
+	void Detach();
+
+	virtual void FlushLine(const FString& Line) override;
+
+private:
+	TArray<FString> BufferedLines;
+	TSharedPtr<FLineBasedTextWriter> Inner;
+};
+
+class FPrefixedTextWriter : public FLineBasedTextWriter
+{
+public:
+	FPrefixedTextWriter(const TCHAR* InPrefix, TSharedRef<FLineBasedTextWriter> InInner);
+	virtual ~FPrefixedTextWriter() override;
+	virtual void FlushLine(const FString& Line) override;
+
+private:
+	FString Prefix;
+	TSharedRef<FLineBasedTextWriter> Inner;
+};
+
+class FBoundedLogWriter : public FLineBasedTextWriter
+{
+public:
+	FBoundedLogWriter(const TCHAR* InFileName, int32 InMaxSize);
+	virtual ~FBoundedLogWriter() override;
+	virtual void FlushLine(const FString& Line) override;
+
+private:
+	FString FileName;
+	FString BackupFileName;
+	int32 MaxSize;
+	FArchive* Inner;
+
+	void OpenInner();
+	void CloseInner();
+};
+
+class FComposedTextWriter : public FLineBasedTextWriter
+{
+public:
+	FComposedTextWriter(const TArray<TSharedRef<FLineBasedTextWriter>>& InInners);
+	virtual ~FComposedTextWriter() override;
+	virtual void FlushLine(const FString& Line) override;
+
+private:
+	TArray<TSharedRef<FLineBasedTextWriter>> Inners;
+};
+
+class FProgressValue
+{
+public:
+	FProgressValue();
+	~FProgressValue();
+
+	void Clear();
+	TTuple<FString, float> GetCurrent() const;
+	void Set(const TCHAR* Message);
+	void Set(const TCHAR* Message, float Fraction);
+	void Set(float Fraction);
+	void Increment(float Fraction);
+	void Push(float MaxFraction);
+	void Pop();
+
+private:
+	TTuple<FString, float> State;
+	TArray<TTuple<float, float>> Ranges;
+
+	float RelativeToAbsoluteFraction(float Fraction);
+};
+
+class FProgressTextWriter : public FLineBasedTextWriter
+{
+public:
+	FProgressTextWriter(const FProgressValue& InValue, const TSharedRef<FLineBasedTextWriter>& InInner);
+	virtual ~FProgressTextWriter() override;
+
+	virtual void FlushLine(const FString& Line) override;
+
+private:
+	static const FString DirectivePrefix;
+
+	FProgressValue Value;
+	TSharedRef<FLineBasedTextWriter> Inner;
+
+	void ProcessInternal(const FString& Line, bool& bSkipLine);
+	static TArray<FString> ParseTokens(const FString& Line);
+	static bool ReadFraction(const TArray<FString>& Tokens, int& TokenIdx, float& Fraction);
+};
diff --git a/Engine/Source/Programs/SlateUGS/Private/UGSCore/Perforce.cpp b/Engine/Source/Programs/SlateUGS/Private/UGSCore/Perforce.cpp
new file mode 100644
--- /dev/null
+++ b/Engine/Source/Programs/SlateUGS/Private/UGSCore/Perforce.cpp
@@ -0,0 +1,1559 @@
+// Copyright 1998-2017 Epic Games, Inc. All Rights Reserved.
+
+#include "Perforce.h"
+#include "Misc/FileHelper.h"
+#include "Misc/Paths.h"
+#include "Algo/Sort.h"
+#include "Algo/Transform.h"
+#include "HAL/PlatformProcess.h"
+#include "HAL/Event.h"
+#include "Utility.h"
+
+bool TryGetValue(const TMap<FString, FString>& Map, const TCHAR* Key, FString& OutValue)
+{
+	const FString* Value = Map.Find(Key);
+	if(Value == nullptr)
+	{
+		OutValue = FString();
+		return false;
+	}
+	else
+	{
+		OutValue = *Value;
+		return true;
+	}
+}
+
+TArray<FString> Split(const FString& Text, const TCHAR* Delim, bool bRemoveEmptyEntries = false)
+{
+	TArray<FString> Result;
+	Text.ParseIntoArray(Result, Delim, bRemoveEmptyEntries);
+	return Result;
+}
+
+FString GetTempFileName()
+{
+	return FPaths::EngineIntermediateDir() / FGuid::NewGuid().ToString();
+}
+
+#if PLATFORM_WINDOWS
+const TCHAR *PerforceExe = TEXT("C:\\Program Files (x86)\\Perforce\\p4.exe");
+#else
+const TCHAR *PerforceExe = TEXT("p4");
+#endif
+
+//// FPerforceChangeSummary ////
+
+FPerforceChangeSummary::FPerforceChangeSummary()
+	: Number(0)
+	, Date(0)
+{
+}
+
+//// FPerforceFileChangeSummary ////
+
+FPerforceFileChangeSummary::FPerforceFileChangeSummary()
+	: Revision(0)
+	, ChangeNumber(0)
+	, Date(0)
+{
+}
+
+//// FPerforceClientRecord ////
+
+FPerforceClientRecord::FPerforceClientRecord(const TMap<FString, FString>& Tags)
+{
+	TryGetValue(Tags, TEXT("client"), Name);
+	TryGetValue(Tags, TEXT("Owner"), Owner);
+	TryGetValue(Tags, TEXT("Host"), Host);
+	TryGetValue(Tags, TEXT("Root"), Root);
+}
+
+//// FPerforceDescribeFileRecord ////
+
+FPerforceDescribeFileRecord::FPerforceDescribeFileRecord()
+	: Revision(0)
+	, FileSize(0)
+{
+}
+
+//// FPerforceDescribeRecord ////
+
+FPerforceDescribeRecord::FPerforceDescribeRecord(const TMap<FString, FString>& Tags)
+{
+	FString ChangeString;
+	if(!TryGetValue(Tags, TEXT("change"), ChangeString) || !FUtility::TryParse(*ChangeString, ChangeNumber))
+	{
+		ChangeNumber = -1;
+	}
+
+	TryGetValue(Tags, TEXT("user"), User);
+	TryGetValue(Tags, TEXT("client"), Client);
+
+	FString TimeString;
+	if(!TryGetValue(Tags, TEXT("time"), TimeString) || !FUtility::TryParse(*TimeString, Time))
+	{
+		Time = -1;
+	}
+
+	TryGetValue(Tags, TEXT("desc"), Description);
+	TryGetValue(Tags, TEXT("status"), Status);
+	TryGetValue(Tags, TEXT("changeType"), ChangeType);
+	TryGetValue(Tags, TEXT("path"), Path);
+
+	for(int Idx = 0;;Idx++)
+	{
+		FString Suffix = FString::Printf(TEXT("%d"), Idx);
+
+		FPerforceDescribeFileRecord File;
+		if(!TryGetValue(Tags, *(FString(TEXT("depotFile")) + Suffix), File.DepotFile))
+		{
+			break;
+		}
+
+		TryGetValue(Tags, *(FString(TEXT("action")) + Suffix), File.Action);
+		TryGetValue(Tags, *(FString(TEXT("type")) + Suffix), File.Type);
+
+		FString RevisionString;
+		if(!TryGetValue(Tags, *(FString(TEXT("rev")) + Suffix), RevisionString) || !FUtility::TryParse(*RevisionString, File.Revision))
+		{
+			File.Revision = -1;
+		}
+
+		FString FileSizeString;
+		if(!TryGetValue(Tags, *(FString(TEXT("fileSize")) + Suffix), FileSizeString) || !FUtility::TryParse(*FileSizeString, File.FileSize))
+		{
+			File.FileSize = -1;
+		}
+
+		TryGetValue(Tags, *(FString(TEXT("digest")) + Suffix), File.Digest);
+		Files.Add(File);
+	}
+}
+
+//// FPerforceInfoRecord ////
+
+FPerforceInfoRecord::FPerforceInfoRecord(const TMap<FString, FString>& Tags)
+	: ServerTimeZone(0, 0, 0)
+{
+	TryGetValue(Tags, TEXT("userName"), UserName);
+	TryGetValue(Tags, TEXT("clientHost"), HostName);
+	TryGetValue(Tags, TEXT("clientAddress"), ClientAddress);
+	TryGetValue(Tags, TEXT("serverAddress"), ServerAddress);
+
+	FString ServerDateTime;
+	if(TryGetValue(Tags, TEXT("serverDate"), ServerDateTime))
+	{
+		TArray<FString> Fields;
+		ServerDateTime.ParseIntoArray(Fields, TEXT(" "));
+		if(Fields.Num() >= 3)
+		{
+			int32 Offset;
+			if(FUtility::TryParse(*Fields[2], Offset))
+			{
+				ServerTimeZone = (Offset < 0)? -FTimespan(-Offset / 100, -Offset % 100, 0) : FTimespan(Offset / 100, Offset % 100, 0);
+			}
+		}
+	}
+}
+
+//// FPerforceFileRecord ////
+
+FPerforceFileRecord::FPerforceFileRecord(const TMap<FString, FString>& Tags)
+	: Revision(0)
+{
+	TryGetValue(Tags, TEXT("depotFile"), DepotPath);
+	TryGetValue(Tags, TEXT("clientFile"), ClientPath);
+	TryGetValue(Tags, TEXT("path"), Path);
+	if(!TryGetValue(Tags, TEXT("action"), Action))
+	{
+		TryGetValue(Tags, TEXT("headAction"), Action);
+	}
+	IsMapped = Tags.Contains(TEXT("isMapped"));
+	Unmap = Tags.Contains(TEXT("unmap"));
+
+	FString RevisionString;
+	if(TryGetValue(Tags, TEXT("rev"), RevisionString))
+	{
+		FUtility::TryParse(*RevisionString, Revision);
+	}
+}
+
+//// FPerforceTagRecordParser ////
+
+FPerforceTagRecordParser::FPerforceTagRecordParser(TFunction<void (const TMap<FString, FString>&)> InOutputRecord)
+	: OutputRecord(InOutputRecord)
+{
+}
+
+FPerforceTagRecordParser::~FPerforceTagRecordParser()
+{
+	Flush();
+}
+
+void FPerforceTagRecordParser::OutputLine(const FString& Line)
+{
+	int SpaceIdx;
+	if(!Line.FindChar(TEXT(' '), SpaceIdx))
+	{
+		SpaceIdx = INDEX_NONE;
+	}
+
+	FString Key = (SpaceIdx > 0)? Line.Left(SpaceIdx) : Line;
+	if(Tags.Contains(Key))
+	{
+		OutputRecord(Tags);
+		Tags.Empty();
+	}
+
+	Tags.FindOrAdd(Key) = (SpaceIdx > 0)? Line.Mid(SpaceIdx + 1) : TEXT("");
+}
+
+void FPerforceTagRecordParser::Flush()
+{
+	if(Tags.Num() > 0)
+	{
+		OutputRecord(Tags);
+		Tags.Empty();
+	}
+}
+
+//// FPerforceSyncOptions ////
+
+FPerforceSyncOptions::FPerforceSyncOptions()
+	: NumRetries(0)
+	, NumThreads(0)
+	, TcpBufferSize(0)
+{
+}
+
+//// FPerforceSpec ////
+
+const TCHAR* FPerforceSpec::GetField(const TCHAR* Name) const
+{
+	for(const TTuple<FString, FString>& Section : Sections)
+	{
+		if(Section.Key == Name)
+		{
+			return *Section.Value;
+		}
+	}
+	return nullptr;
+}
+
+void FPerforceSpec::SetField(const TCHAR* Name, const TCHAR* Value)
+{
+	for(TTuple<FString, FString>& Section : Sections)
+	{
+		if(Section.Key == Name)
+		{
+			Section.Value = Value;
+			return;
+		}
+	}
+	Sections.Add(TTuple<FString, FString>(Name, Value));
+}
+
+bool FPerforceSpec::TryParse(const TArray<FString>& Lines, TSharedPtr<FPerforceSpec>& OutSpec, FOutputDevice& Log)
+{
+	TSharedPtr<FPerforceSpec> Spec = MakeShared<FPerforceSpec>();
+	for(int LineIdx = 0; LineIdx < Lines.Num(); LineIdx++)
+	{
+		FString Line = Lines[LineIdx].TrimStartAndEnd();
+		if(Line.Len() > 0 && Lines[LineIdx][0] != '#')
+		{
+			// Read the section name
+			int SeparatorIdx;
+			if(!Lines[LineIdx].FindChar(':', SeparatorIdx))
+			{
+				SeparatorIdx = -1;
+			}
+			if(SeparatorIdx == -1 || !FChar::IsAlpha(Lines[LineIdx][0]))
+			{
+				Log.Logf(TEXT("Invalid spec format at line %d: \"%s\""), LineIdx, *Lines[LineIdx]);
+				return false;
+			}
+
+			// Get the section name
+			FString SectionName = Lines[LineIdx].Mid(0, SeparatorIdx);
+
+			// Parse the section value
+			FString Value = Lines[LineIdx].Mid(SeparatorIdx + 1).TrimStartAndEnd();
+			for(; LineIdx + 1 < Lines.Num(); LineIdx++)
+			{
+				if(Lines[LineIdx + 1].Len() == 0)
+				{
+					Value += TEXT("\n");
+				}
+				else if(Lines[LineIdx + 1][0] == '\t')
+				{
+					Value += Lines[LineIdx + 1].Mid(1) + TEXT("\n");
+				}
+				else
+				{
+					break;
+				}
+			}
+			Value.TrimEndInline();
+
+			Spec->Sections.Add(TTuple<FString,FString>(SectionName, Value));
+		}
+	}
+	OutSpec = Spec;
+	return true;
+}
+
+FString FPerforceSpec::ToString() const 
+{
+	FString Result;
+	for(const TTuple<FString, FString>& Section : Sections)
+	{
+		if(Section.Value.Contains(TEXT("\n")))
+		{
+			Result += Section.Key + TEXT(":\n\t") + Section.Value.Replace(TEXT("\n"), TEXT("\n\t")) + LINE_TERMINATOR;
+		}
+		else
+		{
+			Result += Section.Key + TEXT(":\t") + Section.Value + LINE_TERMINATOR;
+		}
+		Result += LINE_TERMINATOR;
+	}
+	return Result;
+}
+
+//// FPerforceOutputLine ////
+
+FPerforceOutputLine::FPerforceOutputLine()
+	: Channel(EPerforceOutputChannel::Unknown)
+{
+}
+
+FPerforceOutputLine::FPerforceOutputLine(EPerforceOutputChannel InChannel, const FString& InText)
+	: Channel(InChannel)
+	, Text(InText)
+{
+}
+
+//// FPerforceConnection ////
+
+FPerforceConnection::FPerforceConnection(const TCHAR* InUserName, const TCHAR* InClientName, const TCHAR* InServerAndPort)
+	: ServerAndPort((InServerAndPort != nullptr)? InServerAndPort : TEXT(""))
+	, UserName((InUserName != nullptr)? InUserName : TEXT(""))
+	, ClientName((InClientName != nullptr)? InClientName : TEXT(""))
+{
+}
+
+TSharedRef<FPerforceConnection> FPerforceConnection::OpenClient(const FString& NewClientName) const
+{
+	return MakeShared<FPerforceConnection>(*UserName, *NewClientName, *ServerAndPort);
+}
+
+bool FPerforceConnection::Info(TSharedPtr<FPerforceInfoRecord>& OutInfo, FEvent* AbortEvent, FOutputDevice& Log) const
+{
+	TArray<TMap<FString, FString>> TagRecords;
+	if(!RunCommand(TEXT("info -s"), TagRecords, ECommandOptions::NoClient, AbortEvent, Log) || TagRecords.Num() != 1)
+	{
+		OutInfo = TSharedPtr<FPerforceInfoRecord>();
+		return false;
+	}
+	else
+	{
+		OutInfo = MakeShared<FPerforceInfoRecord>(TagRecords[0]);
+		return true;
+	}
+}
+
+bool FPerforceConnection::GetSetting(const FString& Name, FString& Value, FEvent* AbortEvent, FOutputDevice& Log) const
+{
+	TArray<FString> Lines;
+	if(!RunCommand(FString::Printf(TEXT("set %s"), *Name), Lines, ECommandOptions::NoChannels, AbortEvent, Log) || Lines.Num() != 1)
+	{
+		Value = FString();
+		return false;
+	}
+	if(Lines[0].Len() <= Name.Len() || !Lines[0].StartsWith(Name) || Lines[0][Name.Len()] != '=')
+	{
+		Value = FString();
+		return false;
+	}
+
+	Value = Lines[0].Mid(Name.Len() + 1);
+
+	int EndIdx = Value.Find(TEXT(" ("));
+	if(EndIdx != INDEX_NONE)
+	{
+		Value = Value.Mid(0, EndIdx).TrimStartAndEnd();
+	}
+	return true;
+}
+
+bool FPerforceConnection::FindClients(TArray<FPerforceClientRecord>& Clients, FEvent* AbortEvent, FOutputDevice& Log) const
+{
+	return RunCommand(TEXT("clients"), Clients, ECommandOptions::NoClient, AbortEvent, Log);
+}
+
+bool FPerforceConnection::FindClients(const FString& ForUserName, TArray<FString>& ClientNames, FEvent* AbortEvent, FOutputDevice& Log) const
+{
+	TArray<FString> Lines;
+	if(!RunCommand(FString::Printf(TEXT("clients -u%s"), *ForUserName), Lines, ECommandOptions::None, AbortEvent, Log))
+	{
+		ClientNames.Empty();
+		return false;
+	}
+
+	for(const FString& Line : Lines)
+	{
+		TArray<FString> Tokens = Split(Line, TEXT(" "), true);
+		if(Tokens.Num() < 5 || Tokens[0] != TEXT("Client") || Tokens[3] != TEXT("root"))
+		{
+			Log.Logf(TEXT("Couldn't parse client from line '%s'"), *Line);
+		}
+		else
+		{
+			ClientNames.Add(Tokens[1]);
+		}
+	}
+	return true;
+}
+
+bool FPerforceConnection::TryGetClientSpec(const FString& InClientName, TSharedPtr<FPerforceSpec>& OutSpec, FEvent* AbortEvent, FOutputDevice& Log) const
+{
+	TArray<FString> Lines;
+	if(!RunCommand(FString::Printf(TEXT("client -o %s"), *InClientName), Lines, ECommandOptions::None, AbortEvent, Log))
+	{
+		OutSpec = nullptr;
+		return false;
+	}
+	if(!FPerforceSpec::TryParse(Lines, OutSpec, Log))
+	{
+		OutSpec = nullptr;
+		return false;
+	}
+	return true;
+}
+
+bool FPerforceConnection::TryGetStreamSpec(const FString& StreamName, TSharedPtr<FPerforceSpec>& OutSpec, FEvent* AbortEvent, FOutputDevice& Log) const
+{
+	TArray<FString> Lines;
+	if(!RunCommand(FString::Printf(TEXT("stream -o %s"), *StreamName), Lines, ECommandOptions::None, AbortEvent, Log))
+	{
+		OutSpec = nullptr;
+		return false;
+	}
+	if(!FPerforceSpec::TryParse(Lines, OutSpec, Log))
+	{
+		OutSpec = nullptr;
+		return false;
+	}
+	return true;
+}
+
+bool FPerforceConnection::FindFiles(const FString& Filter, TArray<FPerforceFileRecord>& OutFileRecords, FEvent* AbortEvent, FOutputDevice& Log) const
+{
+	return RunCommand(FString::Printf(TEXT("fstat \"%s\""), *Filter), OutFileRecords, ECommandOptions::None, AbortEvent, Log);
+}
+
+bool FPerforceConnection::Print(const FString& DepotPath, TArray<FString>& OutLines, FEvent* AbortEvent, FOutputDevice& Log) const
+{
+	FString TempFileName = GetTempFileName();
+	if(!PrintToFile(DepotPath, TempFileName, AbortEvent, Log))
+	{
+		return false;
+	}
+	else
+	{
+		return FFileHelper::LoadFileToStringArray(OutLines, *TempFileName);
+	}
+}
+
+bool FPerforceConnection::PrintToFile(const FString& DepotPath, const FString& OutputFileName, FEvent* AbortEvent, FOutputDevice& Log) const
+{
+	return RunCommand(FString::Printf(TEXT("print -q -o \"%s\" \"%s\""), *OutputFileName, *DepotPath), ECommandOptions::None, AbortEvent, Log);
+}
+
+bool FPerforceConnection::FileExists(const FString& Filter, bool& bExists, FEvent* AbortEvent, FOutputDevice& Log) const
+{
+	TArray<FPerforceFileRecord> FileRecords;
+	if(RunCommand(FString::Printf(TEXT("fstat \"%s\""), *Filter), FileRecords, ECommandOptions::IgnoreNoSuchFilesError | ECommandOptions::IgnoreFilesNotInClientViewError, AbortEvent, Log))
+	{
+		bExists = (FileRecords.Num() > 0);
+		return true;
+	}
+	else
+	{
+		bExists = false;
+		return false;
+	}
+}
+
+bool FPerforceConnection::FindChanges(const FString& Filter, int MaxResults, TArray<FPerforceChangeSummary>& OutChanges, FEvent* AbortEvent, FOutputDevice& Log) const
+{
+	TArray<FString> Filters;
+	Filters.Add(Filter);
+	return FindChanges(Filters, MaxResults, OutChanges, AbortEvent, Log);
+}
+
+bool FPerforceConnection::FindChanges(const TArray<FString>& Filters, int MaxResults, TArray<FPerforceChangeSummary>& OutChanges, FEvent* AbortEvent, FOutputDevice& Log) const
+{
+	OutChanges.Empty();
+
+	FString Arguments = TEXT("changes -s submitted -t -L");
+	if(MaxResults > 0)
+	{
+		Arguments += FString::Printf(TEXT(" -m %d"), MaxResults);
+	}
+	for(const FString& Filter : Filters)
+	{
+		Arguments += FString::Printf(TEXT(" \"%s\""), *Filter);
+	}
+
+	TArray<FString> Lines;
+	if(!RunCommand(Arguments, Lines, ECommandOptions::None, AbortEvent, Log))
+	{
+		return false;
+	}
+
+	for(int Idx = 0; Idx < Lines.Num(); Idx++)
+	{
+		FPerforceChangeSummary Change;
+		if(!TryParseChangeSummary(Lines[Idx], Change))
+		{
+			Log.Logf(TEXT("Couldn't parse description from '%s'"), *Lines[Idx]);
+		}
+		else
+		{
+			for(; Idx + 1 < Lines.Num(); Idx++)
+			{
+				if(Lines[Idx + 1].Len() == 0)
+				{
+					Change.Description += TEXT("\n");
+				}
+				else if(Lines[Idx + 1].StartsWith("\t"))
+				{
+					Change.Description += Lines[Idx + 1].Mid(1) + TEXT("\n");
+				}
+				else
+				{
+					break;
+				}
+			}
+			Change.Description = Change.Description.TrimStartAndEnd();
+
+			OutChanges.Add(Change);
+		}
+	}
+
+	Algo::Sort(OutChanges, [](const FPerforceChangeSummary& A, const FPerforceChangeSummary& B){ return A.Number > B.Number; });
+
+	for(int Idx = OutChanges.Num() - 1; Idx > 0; Idx--)
+	{
+		if(OutChanges[Idx - 1].Number == OutChanges[Idx].Number)
+		{
+			OutChanges.RemoveAt(Idx);
+		}
+	}
+
+	if(MaxResults >= 0 && MaxResults < OutChanges.Num())
+	{
+		OutChanges.RemoveAt(MaxResults, OutChanges.Num() - MaxResults);
+	}
+
+	return true;
+}
+
+bool FPerforceConnection::TryParseChangeSummary(const FString& Line, FPerforceChangeSummary& Change) const
+{
+	TArray<FString> Tokens = Split(Line, TEXT(" "), true);
+	if(Tokens.Num() == 7 && Tokens[0] == "Change" && Tokens[2] == "on" && Tokens[5] == "by")
+	{
+		if(FUtility::TryParse(*Tokens[1], Change.Number) && FPerforceUtils::TryParseDateTime(Tokens[3], Tokens[4], Change.Date))
+		{
+			int UserClientIdx;
+			if(Tokens[6].FindChar(TEXT('@'), UserClientIdx))
+			{
+				Change.User = Tokens[6].Mid(0, UserClientIdx);
+				Change.Client = Tokens[6].Mid(UserClientIdx + 1);
+				return true;
+			}
+		}
+	}
+	return false;
+}
+
+bool FPerforceConnection::FindFileChanges(const FString& FilePath, int MaxResults, TArray<FPerforceFileChangeSummary>& Changes, FEvent* AbortEvent, FOutputDevice& Log) const
+{
+	FString Arguments = TEXT("filelog -L -t");
+	if(MaxResults > 0)
+	{
+		Arguments += FString::Printf(TEXT(" -m %d"), MaxResults);
+	}
+	Arguments += FString::Printf(TEXT(" \"%s\""), *FilePath);
+
+	TArray<FString> Lines;
+	if(!RunCommand(Arguments, EPerforceOutputChannel::TaggedInfo, Lines, ECommandOptions::None, AbortEvent, Log))
+	{
+		return false;
+	}
+
+	for(int Idx = 0; Idx < Lines.Num(); Idx++)
+	{
+		FPerforceFileChangeSummary Change;
+		if(!TryParseFileChangeSummary(Lines, Idx, Change))
+		{
+			Log.Logf(TEXT("Couldn't parse description from '%s'"), *Lines[Idx]);
+		}
+		else
+		{
+			Changes.Add(Change);
+		}
+	}
+
+	return true;
+}
+
+bool FPerforceConnection::TryParseFileChangeSummary(const TArray<FString>& Lines, int& LineIdx, FPerforceFileChangeSummary& OutChange) const
+{
+	TArray<FString> Tokens = Split(Lines[LineIdx], TEXT(" "), true);
+	if(Tokens.Num() != 10 || !Tokens[0].StartsWith(TEXT("#")) || Tokens[1] != TEXT("change") || Tokens[4] != TEXT("on") || Tokens[7] != TEXT("by"))
+	{
+		return false;
+	}
+
+	if(!FUtility::TryParse(*Tokens[0].Mid(1), OutChange.Revision) || !FUtility::TryParse(*Tokens[2], OutChange.ChangeNumber) || ensure(false))//!DateTime.TryParse(Tokens[5] + " " + Tokens[6], out Change.Date))
+	{
+		return false;
+	}
+
+	int UserClientIdx;
+	if(Tokens[8].FindChar('@', UserClientIdx))
+	{
+		return false;
+	}
+
+	OutChange.Action = Tokens[3];
+	OutChange.Type = Tokens[9];
+	if(OutChange.Type.StartsWith(TEXT("(")))
+	{
+		OutChange.Type = OutChange.Type.Mid(1);
+	}
+	if(OutChange.Type.EndsWith(TEXT(")")))
+	{
+		OutChange.Type = OutChange.Type.Left(OutChange.Type.Len() - 1);
+	}
+	OutChange.User = Tokens[8].Mid(0, UserClientIdx);
+	OutChange.Client = Tokens[8].Mid(UserClientIdx + 1);
+
+	FString Description;
+	for(; LineIdx + 1 < Lines.Num(); LineIdx++)
+	{
+		if(Lines[LineIdx + 1].Len() == 0)
+		{
+			Description += TEXT("\n");
+		}
+		else if(Lines[LineIdx + 1].StartsWith("\t"))
+		{
+			Description += Lines[LineIdx + 1].Mid(1) + TEXT("\n");
+		}
+		else
+		{
+			break;
+		}
+	}
+	Description.TrimStartAndEndInline();
+
+	OutChange.Description = Description;
+
+	return true;
+}
+
+bool FPerforceConnection::ConvertToClientPath(const FString& FileName, FString& OutClientFileName, FEvent* AbortEvent, FOutputDevice& Log) const
+{
+	TSharedPtr<FPerforceWhereRecord> WhereRecord;
+	if(Where(FileName, WhereRecord, AbortEvent, Log))
+	{
+		OutClientFileName = WhereRecord->ClientPath;
+		return true;
+	}
+	else
+	{
+		OutClientFileName = TEXT("");
+		return false;
+	}
+}
+
+bool FPerforceConnection::ConvertToDepotPath(const FString& FileName, FString& OutDepotFileName, FEvent* AbortEvent, FOutputDevice& Log) const
+{
+	TSharedPtr<FPerforceWhereRecord> WhereRecord;
+	if(Where(FileName, WhereRecord, AbortEvent, Log))
+	{
+		OutDepotFileName = WhereRecord->DepotPath;
+		return true;
+	}
+	else
+	{
+		OutDepotFileName = TEXT("");
+		return false;
+	}
+}
+
+bool FPerforceConnection::ConvertToLocalPath(const FString& FileName, FString& OutLocalFileName, FEvent* AbortEvent, FOutputDevice& Log) const
+{
+	TSharedPtr<FPerforceWhereRecord> WhereRecord;
+	if(Where(FileName, WhereRecord, AbortEvent, Log))
+	{
+		OutLocalFileName = FUtility::GetPathWithCorrectCase(*WhereRecord->LocalPath);
+		return true;
+	}
+	else
+	{
+		OutLocalFileName = TEXT("");
+		return false;
+	}
+}
+
+bool FPerforceConnection::FindStreams(const FString& Filter, TArray<FString>& OutStreamNames, FEvent* AbortEvent, FOutputDevice& Log) const
+{
+	TArray<FString> Lines;
+	if(!RunCommand(FString::Printf(TEXT("streams -F \"Stream=%s\""), *Filter), Lines, ECommandOptions::None, AbortEvent, Log))
+	{
+		OutStreamNames.Empty();
+		return false;
+	}
+
+	TArray<FString> StreamNames;
+	for(const FString& Line : Lines)
+	{
+		TArray<FString> Tokens = Split(Line, TEXT(" "), true);
+		if(Tokens.Num() < 2 || Tokens[0] != TEXT("Stream") || !Tokens[1].StartsWith(TEXT("//")))
+		{
+			Log.Logf(TEXT("Unexpected output from stream query: %s"), *Line);
+			OutStreamNames.Empty();
+			return false;
+		}
+		StreamNames.Add(Tokens[1]);
+	}
+	OutStreamNames = StreamNames;
+
+	return true;
+}
+
+bool FPerforceConnection::HasOpenFiles(FEvent* AbortEvent, FOutputDevice& Log) const
+{
+	TArray<FPerforceFileRecord> Records;
+	bool bResult = RunCommand(TEXT("opened -m 1"), Records, ECommandOptions::None, AbortEvent, Log);
+	return bResult && Records.Num() > 0;
+}
+
+bool FPerforceConnection::SwitchStream(const FString& NewStream, FEvent* AbortEvent, FOutputDevice& Log) const
+{
+	return RunCommand(FString::Printf(TEXT("client -f -s -S \"%s\" \"%s\""), *NewStream, *ClientName), ECommandOptions::None, AbortEvent, Log);
+}
+
+bool FPerforceConnection::Describe(int ChangeNumber, TSharedPtr<FPerforceDescribeRecord>& OutRecord, FEvent* AbortEvent, FOutputDevice& Log) const
+{
+	FString CommandLine = FString::Printf(TEXT("describe -s %d"), ChangeNumber);
+
+	TArray<TMap<FString, FString>> Records;
+	if(!RunCommandWithBinaryOutput(CommandLine, Records, ECommandOptions::None, AbortEvent, Log))
+	{
+		OutRecord = nullptr;
+		return false;
+	}
+	if(Records.Num() != 1)
+	{
+		Log.Logf(TEXT("Expected 1 record from p4 %s, got %d"), *CommandLine, Records.Num());
+		OutRecord = nullptr;
+		return false;
+	}
+
+	FString Code;
+	if(!TryGetValue(Records[0], TEXT("code"), Code) || Code != "stat")
+	{
+		Log.Logf(TEXT("Unexpected response from p4 %s (code=%s)"), *CommandLine, *Code);
+		OutRecord = nullptr;
+		return false;
+	}
+
+	OutRecord = MakeShared<FPerforceDescribeRecord>(Records[0]);
+	return true;
+}
+
+bool FPerforceConnection::Where(const FString& Filter, TSharedPtr<FPerforceWhereRecord>& OutWhereRecord, FEvent* AbortEvent, FOutputDevice& Log) const
+{
+	TArray<FPerforceFileRecord> FileRecords;
+	if(!RunCommand(FString::Printf(TEXT("where \"%s\""), *Filter), FileRecords, ECommandOptions::None, AbortEvent, Log))
+	{
+		OutWhereRecord = nullptr;
+		return false;
+	}
+
+	FileRecords.RemoveAll([](const FPerforceFileRecord& Record){ return Record.Unmap; });
+
+	if(FileRecords.Num() == 0)
+	{
+		Log.Logf(TEXT("'%s' is not mapped to workspace."), *Filter);
+		OutWhereRecord = nullptr;
+		return false;
+	}
+	else if(FileRecords.Num() > 1)
+	{
+		Log.Logf(TEXT("File is mapped to %d locations:"), FileRecords.Num());
+		for(const FPerforceFileRecord& FileRecord : FileRecords)
+		{
+			Log.Logf(TEXT("  %s"), *FileRecord.Path);
+		}
+		OutWhereRecord = nullptr;
+		return false;
+	}
+
+	OutWhereRecord = MakeShared<FPerforceWhereRecord>();
+	OutWhereRecord->LocalPath = FileRecords[0].Path;
+	OutWhereRecord->DepotPath = FileRecords[0].DepotPath;
+	OutWhereRecord->ClientPath = FileRecords[0].ClientPath;
+	return true;
+}
+
+bool FPerforceConnection::Have(const FString& Filter, TArray<FPerforceFileRecord>& OutFileRecords, FEvent* AbortEvent, FOutputDevice& Log) const
+{
+	return RunCommand(FString::Printf(TEXT("have \"%s\""), *Filter), OutFileRecords, ECommandOptions::None, AbortEvent, Log);
+}
+
+bool FPerforceConnection::Stat(const FString& Filter, TArray<FPerforceFileRecord>& OutFileRecords, FEvent* AbortEvent, FOutputDevice& Log) const
+{
+	return RunCommand(FString::Printf(TEXT("fstat \"%s\""), *Filter), OutFileRecords, ECommandOptions::None, AbortEvent, Log);
+}
+
+bool FPerforceConnection::Sync(const FString& Filter, FEvent* AbortEvent, FOutputDevice& Log) const
+{
+	return RunCommand(FString(TEXT("sync ")) + Filter, ECommandOptions::IgnoreFilesUpToDateError, AbortEvent, Log);
+}
+
+bool FPerforceConnection::Sync(const TArray<FString>& DepotPaths, int ChangeNumber, TFunction<void(const FPerforceFileRecord&)> SyncOutput, TArray<FString>& OutTamperedFiles, const FPerforceSyncOptions* Options, FEvent* AbortEvent, FOutputDevice& Log) const
+{
+	// Write all the files we want to sync to a temp file
+	TArray<FString> SyncFiles;
+	for(const FString& DepotPath : DepotPaths)
+	{
+		SyncFiles.Add(FString::Printf(TEXT("%s@%d"), *DepotPath, ChangeNumber));
+	}
+	FString TempFileName = GetTempFileName();
+	FFileHelper::SaveStringArrayToFile(SyncFiles, *TempFileName);
+
+	// Create a filter to strip all the sync records
+	FPerforceTagRecordParser Parser([SyncOutput](const TMap<FString, FString>& Tags){ SyncOutput(FPerforceFileRecord(Tags)); });
+
+	FString CommandLine;
+	CommandLine += FString::Printf(TEXT("-x \"%s\" -z tag"), *TempFileName);
+	if(Options != nullptr && Options->NumRetries > 0)
+	{
+		CommandLine += FString::Printf(TEXT(" -r %d"), Options->NumRetries);
+	}
+	if(Options != nullptr && Options->TcpBufferSize > 0)
+	{
+		CommandLine += FString::Printf(TEXT(" -v net.tcpsize=%d"), Options->TcpBufferSize);
+	}
+	CommandLine += TEXT(" sync");
+	if(Options != nullptr && Options->NumThreads > 1)
+	{
+		CommandLine += FString::Printf(TEXT(" --parallel=threads=%d"), Options->NumThreads);
+	}
+
+	return RunCommand(CommandLine, nullptr, [&Parser, &OutTamperedFiles, &Log](const FPerforceOutputLine& Line){ return FilterSyncOutput(Line, Parser, OutTamperedFiles, Log); }, ECommandOptions::NoFailOnErrors | ECommandOptions::IgnoreFilesUpToDateError | ECommandOptions::IgnoreExitCode, AbortEvent, Log);
+}
+
+bool FPerforceConnection::FilterSyncOutput(const FPerforceOutputLine& Line, FPerforceTagRecordParser& Parser, TArray<FString>& OutTamperedFiles, FOutputDevice& Log)
+{
+	if(Line.Channel == EPerforceOutputChannel::TaggedInfo)
+	{
+		Parser.OutputLine(Line.Text);
+		return true;
+	}
+
+	Log.Logf(TEXT("%s"), *Line.Text);
+
+	static const FString Prefix = TEXT("Can't clobber writable file ");
+	if(Line.Channel == EPerforceOutputChannel::Error && Line.Text.StartsWith(Prefix))
+	{
+		OutTamperedFiles.Add(Line.Text.Mid(Prefix.Len()).TrimStartAndEnd());
+		return true;
+	}
+
+	return Line.Channel != EPerforceOutputChannel::Error;
+}
+
+void FPerforceConnection::ParseTamperedFile(const FString& Line, TArray<FString>& OutTamperedFiles)
+{
+	static const FString Prefix = TEXT("Can't clobber writable file ");
+	if(Line.StartsWith(Prefix))
+	{
+		OutTamperedFiles.Add(Line.Mid(Prefix.Len()).TrimStartAndEnd());
+	}
+}
+
+bool FPerforceConnection::SyncPreview(const FString& Filter, int ChangeNumber, bool bOnlyFilesInThisChange, TArray<FPerforceFileRecord>& OutFileRecords, FEvent* AbortEvent, FOutputDevice& Log) const
+{
+	return RunCommand(FString::Printf(TEXT("sync -n %s@%s%d"), *Filter, bOnlyFilesInThisChange? TEXT("=") : TEXT(""), ChangeNumber), OutFileRecords, ECommandOptions::IgnoreFilesUpToDateError | ECommandOptions::IgnoreNoSuchFilesError, AbortEvent, Log);
+}
+
+bool FPerforceConnection::ForceSync(const FString& Filter, int ChangeNumber, FEvent* AbortEvent, FOutputDevice& Log) const
+{
+	return RunCommand(FString::Printf(TEXT("sync -f \"%s\"@%d"), *Filter, ChangeNumber), ECommandOptions::IgnoreFilesUpToDateError, AbortEvent, Log);
+}
+
+bool FPerforceConnection::GetOpenFiles(const FString& Filter, TArray<FPerforceFileRecord>& OutFileRecords, FEvent* AbortEvent, FOutputDevice& Log) const
+{
+	return RunCommand(FString::Printf(TEXT("opened \"%s\""), *Filter), OutFileRecords, ECommandOptions::None, AbortEvent, Log);
+}
+
+bool FPerforceConnection::GetUnresolvedFiles(const FString& Filter, TArray<FPerforceFileRecord>& OutFileRecords, FEvent* AbortEvent, FOutputDevice& Log) const
+{
+	return RunCommand(FString::Printf(TEXT("fstat -Ru \"%s\""), *Filter), OutFileRecords, ECommandOptions::IgnoreNoSuchFilesError | ECommandOptions::IgnoreFilesNotOpenedOnThisClientError, AbortEvent, Log);
+}
+
+bool FPerforceConnection::AutoResolveFile(const FString& File, FEvent* AbortEvent, FOutputDevice& Log) const
+{
+	return RunCommand(FString::Printf(TEXT("resolve -am %s"), *File), ECommandOptions::None, AbortEvent, Log);
+}
+
+bool FPerforceConnection::GetActiveStream(FString& OutStreamName, FEvent* AbortEvent, FOutputDevice& Log) const
+{
+	TSharedPtr<FPerforceSpec> ClientSpec;
+	if(TryGetClientSpec(ClientName, ClientSpec, AbortEvent, Log))
+	{
+		OutStreamName = ClientSpec->GetField(TEXT("Stream"));
+		return OutStreamName.Len() > 0;
+	}
+	else
+	{
+		OutStreamName.Empty();
+		return false;
+	}
+}
+
+bool FPerforceConnection::RunCommand(const FString& CommandLine, TArray<FPerforceFileRecord>& OutFileRecords, ECommandOptions Options, FEvent* AbortEvent, FOutputDevice& Log) const
+{
+	TArray<TMap<FString, FString>> TagRecords;
+	if(!RunCommand(CommandLine, TagRecords, Options, AbortEvent, Log))
+	{
+		OutFileRecords.Empty();
+		return false;
+	}
+	else
+	{
+		Algo::Transform(TagRecords, OutFileRecords, [](const TMap<FString, FString>& Tags) -> FPerforceFileRecord { return FPerforceFileRecord(Tags); });
+		return true;
+	}
+}
+
+bool FPerforceConnection::RunCommand(const FString& CommandLine, TArray<FPerforceClientRecord>& OutClientRecords, ECommandOptions Options, FEvent* AbortEvent, FOutputDevice& Log) const
+{
+	TArray<TMap<FString, FString>> TagRecords;
+	if(!RunCommand(CommandLine, TagRecords, Options, AbortEvent, Log))
+	{
+		return false;
+	}
+	else
+	{
+		Algo::Transform(TagRecords, OutClientRecords, [](const TMap<FString, FString>& Tags) -> FPerforceClientRecord { return FPerforceClientRecord(Tags); });
+		return true;
+	}
+}
+
+bool FPerforceConnection::RunCommand(const FString& CommandLine, TArray<TMap<FString, FString>>& OutTagRecords, ECommandOptions Options, FEvent* AbortEvent, FOutputDevice& Log) const
+{
+	TArray<FString> Lines;
+	if(!RunCommand(FString(TEXT("-ztag ")) + CommandLine, EPerforceOutputChannel::TaggedInfo, Lines, Options, AbortEvent, Log))
+	//if(!RunCommand(FString(CommandLine, EPerforceOutputChannel::TaggedInfo, Lines, Options, AbortEvent, Log))
+	{
+		return false;
+	}
+
+	FPerforceTagRecordParser Parser([&OutTagRecords](const TMap<FString, FString>& Tags){ OutTagRecords.Add(Tags); });
+	for(const FString& Line : Lines)
+	{
+		Parser.OutputLine(Line);
+	}
+	return true;
+}
+
+bool FPerforceConnection::RunCommand(const FString& CommandLine, ECommandOptions Options, FEvent* AbortEvent, FOutputDevice& Log) const
+{
+	TArray<FString> Lines;
+	return RunCommand(CommandLine, Lines, Options, AbortEvent, Log);
+}
+
+bool FPerforceConnection::RunCommand(const FString& CommandLine, TArray<FString>& OutLines, ECommandOptions Options, FEvent* AbortEvent, FOutputDevice& Log) const
+{
+	return RunCommand(CommandLine, EPerforceOutputChannel::Info, OutLines, Options, AbortEvent, Log);
+}
+
+int RunPerforceCommand(const FString& CommandLine, TArray<FString>& OutLines)
+{
+	uint32 ProcId;
+	void* ReadPipe = nullptr;
+    void* WritePipe = nullptr;
+    FPlatformProcess::CreatePipe(ReadPipe, WritePipe);
+
+	FProcHandle P4Proc = FPlatformProcess::CreateProc(PerforceExe, *CommandLine, false, false, false, &ProcId, 0, nullptr, WritePipe, ReadPipe);
+
+	FString P4Output;
+    FString LatestOutput = FPlatformProcess::ReadPipe(ReadPipe);
+    while (FPlatformProcess::IsProcRunning(P4Proc) || !LatestOutput.IsEmpty())
+    {
+        P4Output += LatestOutput;
+        LatestOutput = FPlatformProcess::ReadPipe(ReadPipe);
+        FPlatformProcess::Sleep(0);
+    }
+
+    FPlatformProcess::ClosePipe(ReadPipe, WritePipe);
+
+	// TODO check if p4 spits out \n on Windows
+	OutLines = Split(P4Output, TEXT("\n"));
+
+	int ExitCode = -1;
+	bool GotReturnCode = FPlatformProcess::GetProcReturnCode(P4Proc, &ExitCode);
+
+	if (GotReturnCode)
+	{
+		return ExitCode;
+	}
+
+	return -1;
+}
+
+bool FPerforceConnection::RunCommand(const FString& CommandLine, EPerforceOutputChannel Channel, TArray<FString>& OutLines, ECommandOptions Options, FEvent* AbortEvent, FOutputDevice& Log) const
+{
+	FString FullCommandLine = GetFullCommandLine(CommandLine, Options);
+	Log.Logf(TEXT("p4> p4.exe %s"), *FullCommandLine);
+
+	TArray<FString> RawOutputLines;
+
+	int ExitCode = RunPerforceCommand(FullCommandLine, RawOutputLines);
+	if (ExitCode != 0 && !EnumHasAnyFlags(Options, ECommandOptions::IgnoreExitCode))
+	{
+		return false;
+	}
+
+	bool bResult = true;
+	if(EnumHasAnyFlags(Options, ECommandOptions::NoChannels))
+	{
+		OutLines = RawOutputLines;
+	}
+	else
+	{
+		TArray<FString> LocalLines;
+		for(const FString& RawOutputLine : RawOutputLines)
+		{
+			bResult &= ParseCommandOutput(RawOutputLine, [&LocalLines, Channel, &Log](const FPerforceOutputLine& Line){ if(Line.Channel == Channel){ LocalLines.Add(Line.Text); return true; } else { Log.Logf(TEXT("%s"), *Line.Text); return Line.Channel != EPerforceOutputChannel::Error; } }, Options);
+		}
+		OutLines = LocalLines;
+	}
+	return bResult;
+}
+
+bool FPerforceConnection::RunCommand(const FString& CommandLine, const TCHAR* Input, TFunction<bool(const FPerforceOutputLine&)> HandleOutput, ECommandOptions Options, FEvent* AbortEvent, FOutputDevice& Log) const
+{
+	FString FullCommandLine = GetFullCommandLine(CommandLine, Options);
+	Log.Logf(TEXT("p4> p4.exe %s"), *FullCommandLine);
+
+	TArray<FString> RawOutputLines;
+
+	int ExitCode = RunPerforceCommand(FullCommandLine, RawOutputLines);
+
+	// TODO check this vs the old way as things *may* have changed. Before they were handling each line as it was coming
+	// back from the process. here we just spin + collect all the output. Then try to process all the output, then if the
+	// exit code was not good return false still.
+	bool bResult = true;
+	FString Line;
+	for(const FString& RawOutputLine : RawOutputLines)
+	{
+		bResult &= ParseCommandOutput(Line, HandleOutput, Options);
+	}
+
+	if (ExitCode != 0 && !EnumHasAnyFlags(Options, ECommandOptions::IgnoreExitCode))
+	{
+		bResult = false;
+	}
+
+	return bResult;
+}
+
+FString FPerforceConnection::GetFullCommandLine(const FString& CommandLine, ECommandOptions Options) const
+{
+	FString FullCommandLine;
+	if(ServerAndPort.Len() > 0)
+	{
+		FullCommandLine += FString::Printf(TEXT("-p%s "), *ServerAndPort);
+	}
+	if(UserName.Len() > 0)
+	{
+		FullCommandLine += FString::Printf(TEXT("-u%s "), *UserName);
+	}
+	if((Options & ECommandOptions::NoClient) == ECommandOptions::None && ClientName.Len() > 0)
+	{
+		FullCommandLine += FString::Printf(TEXT("-c%s "), *ClientName);
+	}
+	if((Options & ECommandOptions::NoChannels) == ECommandOptions::None)
+	{
+		FullCommandLine += TEXT("-s ");
+	}
+	FullCommandLine += CommandLine;
+
+	return FullCommandLine;
+}
+
+bool FPerforceConnection::ParseCommandOutput(const FString& Text, TFunction<bool(const FPerforceOutputLine&)> HandleOutput, ECommandOptions Options) const
+{
+	if(EnumHasAnyFlags(Options, ECommandOptions::NoChannels))
+	{
+		FPerforceOutputLine Line(EPerforceOutputChannel::Unknown, Text);
+		return HandleOutput(Line);
+	}
+	else if(!IgnoreCommandOutput(Text, Options))
+	{
+		FPerforceOutputLine Line;
+		if(Text.StartsWith(TEXT("text: ")))
+		{
+			Line = FPerforceOutputLine(EPerforceOutputChannel::Text, Text.Mid(6));
+		}
+		else if(Text.StartsWith("info: "))
+		{
+			Line = FPerforceOutputLine(EPerforceOutputChannel::Info, Text.Mid(6));
+		}
+		else if(Text.StartsWith("info1: "))
+		{
+			Line = FPerforceOutputLine(IsValidTag(Text, 7)? EPerforceOutputChannel::TaggedInfo : EPerforceOutputChannel::Info, Text.Mid(7));
+		}
+		else if(Text.StartsWith("warning: "))
+		{
+			Line = FPerforceOutputLine(EPerforceOutputChannel::Warning, Text.Mid(9));
+		}
+		else if(Text.StartsWith("error: "))
+		{
+			Line = FPerforceOutputLine(EPerforceOutputChannel::Error, Text.Mid(7));
+		}
+		else
+		{
+			Line = FPerforceOutputLine(EPerforceOutputChannel::Unknown, Text);
+		}
+		return HandleOutput(Line) && (Line.Channel != EPerforceOutputChannel::Error || EnumHasAnyFlags(Options, ECommandOptions::NoFailOnErrors)) && Line.Channel != EPerforceOutputChannel::Unknown;
+	}
+	return true;
+}
+
+bool FPerforceConnection::IsValidTag(const FString& Line, int StartIndex)
+{
+	// Annoyingly, we sometimes get commentary with an info1: prefix. Since it typically starts with a depot or file path, we can pick it out.
+	for(int Idx = StartIndex; Idx < Line.Len() && Line[Idx] != ' '; Idx++)
+	{
+		if(Line[Idx] == '/' || Line[Idx] == '\\')
+		{
+			return false;
+		}
+	}
+	return true;
+}
+
+bool FPerforceConnection::IgnoreCommandOutput(const FString& Text, ECommandOptions Options)
+{
+	if(Text.StartsWith(TEXT("exit: ")) || Text.StartsWith(TEXT("info2: ")) || Text.Len() == 0)
+	{
+		return true;
+	}
+	else if(EnumHasAnyFlags(Options, ECommandOptions::IgnoreFilesUpToDateError) && Text.StartsWith(TEXT("error: ")) && Text.EndsWith(TEXT("- file(s) up-to-date.")))
+	{
+		return true;
+	}
+	else if(EnumHasAnyFlags(Options, ECommandOptions::IgnoreNoSuchFilesError) && Text.StartsWith(TEXT("error: ")) && Text.EndsWith(TEXT(" - no such file(s).")))
+	{
+		return true;
+	}
+	else if(EnumHasAnyFlags(Options, ECommandOptions::IgnoreFilesNotInClientViewError) && Text.StartsWith(TEXT("error: ")) && Text.EndsWith(TEXT("- file(s) not in client view.")))
+	{
+		return true;
+	}
+	else if(EnumHasAnyFlags(Options, ECommandOptions::IgnoreFilesNotOpenedOnThisClientError) && Text.StartsWith(TEXT("error: ")) && Text.EndsWith(TEXT(" - file(s) not opened on this client.")))
+	{
+		return true;
+	}
+	return false;
+}
+
+const uint8* FPerforceConnection::ReadBinaryField(const uint8* Pos, const uint8* End, FString* OutField)
+{
+	// Read the field type
+	if(End - Pos < sizeof(uint8))
+	{
+		return nullptr;
+	}
+	uint8 KeyFieldType = *Pos;
+	Pos++;
+
+	// Skip over the field data
+	if(KeyFieldType == 's')
+	{
+		// Read the string length
+		if(End - Pos < sizeof(uint32))
+		{
+			return nullptr;
+		}
+		uint32 KeyLength;
+		memcpy(&KeyLength, Pos, sizeof(uint32));
+		Pos += sizeof(uint32);
+
+		// Read the string data
+		if(End - Pos < KeyLength)
+		{
+			return nullptr;
+		}
+		if(OutField != nullptr)
+		{
+			TStringConversion<FUTF8ToTCHAR_Convert> Converter((const ANSICHAR*)Pos, KeyLength);
+			*OutField = FString(Converter.Length(), Converter.Get());
+		}
+		Pos += KeyLength;
+		return Pos;
+	}
+	else if(KeyFieldType == 'i')
+	{
+		// Skip over the integer data
+		if(End - Pos < sizeof(uint32))
+		{
+			return nullptr;
+		}
+		if(OutField != nullptr)
+		{
+			int32 Value;
+			memcpy(&Value, Pos, sizeof(int32));
+			*OutField = FString::Printf(TEXT("%d"), Value);
+		}
+		Pos += sizeof(uint32);
+		return Pos;
+	}
+	else
+	{
+		checkf(false, TEXT("Invalid field type '%d' (%c)"), (int)*Pos, (TCHAR)*Pos);
+		return nullptr;
+	}
+}
+
+const uint8* FPerforceConnection::ReadBinaryRecord(const uint8* Pos, const uint8* End, TMap<FString, FString>* OutRecord)
+{
+	if(End > Pos && *Pos == '{')
+	{
+		Pos++;
+		while(Pos < End)
+		{
+			// If this is the end of the dictionary, return success
+			if(*Pos == '0')
+			{
+				return Pos + 1;
+			}
+
+			// Read or skip over the field data
+			if(OutRecord == nullptr)
+			{
+				Pos = ReadBinaryField(Pos, End, nullptr);
+				if(Pos == nullptr)
+				{
+					return nullptr;
+				}
+
+				Pos = ReadBinaryField(Pos, End, nullptr);
+				if(Pos == nullptr)
+				{
+					return nullptr;
+				}
+			}
+			else
+			{
+				FString Key;
+				Pos = ReadBinaryField(Pos, End, &Key);
+				if(Pos == nullptr)
+				{
+					return nullptr;
+				}
+
+				FString Value;
+				Pos = ReadBinaryField(Pos, End, &Value);
+				if(Pos == nullptr)
+				{
+					return nullptr;
+				}
+
+				OutRecord->FindOrAdd(MoveTemp(Key)) = MoveTemp(Value);
+			}
+		}
+	}
+	return nullptr;
+}
+
+/// <summary>
+/// Execute a Perforce command and parse the output as marshalled Python objects. This is more robustly defined than the text-based tagged output
+/// format, because it avoids ambiguity when returned fields can have newlines.
+/// </summary>
+/// <param name="CommandLine">Command line to execute Perforce with</param>
+/// <param name="TaggedOutput">List that receives the output records</param>
+/// <param name="WithClient">Whether to include client information on the command line</param>
+bool FPerforceConnection::RunCommandWithBinaryOutput(const FString& CommandLine, TArray<TMap<FString, FString>>& OutRecords, ECommandOptions Options, FEvent* AbortEvent, FOutputDevice& Log) const
+{
+	return RunCommandWithBinaryOutput(CommandLine, [&OutRecords](const TMap<FString, FString>& Record){ OutRecords.Add(Record); return true; }, Options, AbortEvent, Log);
+}
+
+/// <summary>
+/// Execute a Perforce command and parse the output as marshalled Python objects. This is more robustly defined than the text-based tagged output
+/// format, because it avoids ambiguity when returned fields can have newlines.
+/// </summary>
+/// <param name="CommandLine">Command line to execute Perforce with</param>
+/// <param name="TaggedOutput">List that receives the output records</param>
+/// <param name="WithClient">Whether to include client information on the command line</param>
+bool FPerforceConnection::RunCommandWithBinaryOutput(const FString& CommandLine, TFunction<void(const TMap<FString, FString>&)> HandleOutput, ECommandOptions Options, FEvent* AbortEvent, FOutputDevice& Log) const
+{
+	FString FullCommandLine(TEXT("-G ") + CommandLine);
+
+	uint32 ProcId;
+	void* ReadPipe = nullptr;
+    void* WritePipe = nullptr;
+    FPlatformProcess::CreatePipe(ReadPipe, WritePipe);
+
+	FProcHandle P4Proc = FPlatformProcess::CreateProc(PerforceExe, *FullCommandLine, false, false, false, &ProcId, 0, nullptr, WritePipe, ReadPipe);
+
+	/*
+	FString P4Output;
+    FString LatestOutput = FPlatformProcess::ReadPipe(ReadPipe);
+    while (FPlatformProcess::IsProcRunning(P4Proc) || !LatestOutput.IsEmpty())
+    {
+        P4Output += LatestOutput;
+        LatestOutput = FPlatformProcess::ReadPipe(ReadPipe);
+        FPlatformProcess::Sleep(0);
+    }
+	*/
+
+	size_t BufferPos = 0;
+	size_t BufferEnd = 0;
+	TArray<uint8> Buffer;
+	TMap<FString, FString> Record;
+	for(;;)
+	{
+		// Determine if the process has exited. Do this before reading, so we'll only exit the loop if no data is read AND the process has exited.
+		bool bHasExited = !FPlatformProcess::IsProcRunning(P4Proc);
+
+		// Check that we don't have an abort event
+		if (AbortEvent->Wait(FTimespan::Zero()))
+		{
+			abort();
+			//throw FAbortException();
+		}
+
+		// Prepare space in the output buffer for the piped data
+		//size_t BytesToRead = 512;
+		//if(Buffer.Num() < BufferEnd + BytesToRead)
+		//{
+			//Buffer.AddUninitialized((BufferEnd + BytesToRead) - Buffer.Num());
+		//}
+
+		// Read data from the child process
+		//size_t BytesRead = Process.Read(Buffer.GetData() + BufferEnd, BytesToRead);
+		TArray<uint8> TempBuffer;
+		bool bRead = FPlatformProcess::ReadPipeToArray(ReadPipe, TempBuffer);
+		size_t BytesRead = TempBuffer.Num();
+
+		// Add new data to the end of the Buffer
+		Buffer += TempBuffer;
+
+		if(BytesRead == 0)
+		{
+			// If it exited, quit. Otherwise sleep until data is available.
+			if(bHasExited)
+			{
+				break;
+			}
+			else if(AbortEvent->Wait(FTimespan::FromMilliseconds(50)))
+			{
+				abort();
+				//throw FAbortException();
+			}
+		}
+		else
+		{
+			// Add the read bytes to the buffer
+			BufferEnd += BytesRead;
+
+			// Process all the records in the buffer that we can
+			const uint8* RecordPos = Buffer.GetData() + BufferPos;
+			const uint8* RecordMaxPos = Buffer.GetData() + BufferEnd;
+			for(;;)
+			{
+				const uint8* RecordEnd = ReadBinaryRecord(RecordPos, RecordMaxPos, &Record);
+				if(RecordEnd == nullptr)
+				{
+					break;
+				}
+				HandleOutput(Record);
+				RecordPos = RecordEnd;
+			}
+			BufferPos = RecordPos - Buffer.GetData();
+
+			// Shrink the buffer down if we've got spare space
+			if(BufferPos > 64)
+			{
+				BufferEnd -= BufferPos;
+				memmove(Buffer.GetData(), Buffer.GetData() + BufferPos, BufferEnd);
+				BufferPos = 0;
+			}
+		}
+	}
+
+    FPlatformProcess::ClosePipe(ReadPipe, WritePipe);
+
+	int ExitCode = -1;
+
+	checkf(BufferPos == BufferEnd, TEXT("%d bytes of incomplete record data received from P4"), BufferEnd - BufferPos);
+	return FPlatformProcess::GetProcReturnCode(P4Proc, &ExitCode) == 0;
+
+	/*
+	FChildProcess Process;
+	if(!Process.Start(PerforceExe, *FString::Printf(TEXT("-G %s"), *CommandLine), nullptr, EChildProcessOptions::OutputPipe, EChildProcessPriority::Normal))
+	{
+		return false;
+	}
+
+	size_t BufferPos = 0;
+	size_t BufferEnd = 0;
+	TArray<uint8> Buffer;
+	TMap<FString, FString> Record;
+	for(;;)
+	{
+		// Determine if the process has exited. Do this before reading, so we'll only exit the loop if no data is read AND the process has exited.
+		bool bHasExited = Process.HasExited();
+
+		// Check that we don't have an abort event
+		if(AbortEvent->Wait(FTimespan::Zero()))
+		{
+			throw FAbortException();
+		}
+
+		// Prepare space in the output buffer for the piped data
+		size_t BytesToRead = 512;
+		if(Buffer.Num() < BufferEnd + BytesToRead)
+		{
+			Buffer.AddUninitialized((BufferEnd + BytesToRead) - Buffer.Num());
+		}
+
+		// Read data from the child process
+		size_t BytesRead = Process.Read(Buffer.GetData() + BufferEnd, BytesToRead);
+		if(BytesRead == 0)
+		{
+			// If it exited, quit. Otherwise sleep until data is available.
+			if(bHasExited)
+			{
+				break;
+			}
+			else if(AbortEvent->Wait(FTimespan::FromMilliseconds(50)))
+			{
+				throw FAbortException();
+			}
+		}
+		else
+		{
+			// Add the read bytes to the buffer
+			BufferEnd += BytesRead;
+
+			// Process all the records in the buffer that we can
+			const uint8* RecordPos = Buffer.GetData() + BufferPos;
+			const uint8* RecordMaxPos = Buffer.GetData() + BufferEnd;
+			for(;;)
+			{
+				const uint8* RecordEnd = ReadBinaryRecord(RecordPos, RecordMaxPos, &Record);
+				if(RecordEnd == nullptr)
+				{
+					break;
+				}
+				HandleOutput(Record);
+				RecordPos = RecordEnd;
+			}
+			BufferPos = RecordPos - Buffer.GetData();
+
+			// Shrink the buffer down if we've got spare space
+			if(BufferPos > 64)
+			{
+				BufferEnd -= BufferPos;
+				memmove(Buffer.GetData(), Buffer.GetData() + BufferPos, BufferEnd);
+				BufferPos = 0;
+			}
+		}
+	}
+	checkf(BufferPos == BufferEnd, TEXT("%d bytes of incomplete record data received from P4"), BufferEnd - BufferPos);
+	return Process.GetExitCode() == 0;
+	*/
+
+	return false;
+}
+
+//// FPerforceUtils ////
+
+FString FPerforceUtils::GetClientOrDepotDirectoryName(const TCHAR* ClientFile)
+{
+	const TCHAR* LastSlash = FCString::Strrchr(ClientFile, '/');
+	if(LastSlash == nullptr)
+	{
+		return "";
+	}
+	else
+	{
+		return FString(LastSlash - ClientFile, ClientFile);
+	}
+}
+
+bool FPerforceUtils::TryParseDateTime(const FString& Date, const FString& Time, FDateTime& OutDate)
+{
+	// Parse the date
+	const TCHAR* Pos = *Date;
+	TCHAR* End = nullptr;
+
+	int32 Year = FCString::Strtoi(Pos, &End, 10);
+	if(End <= Pos || *End != '/')
+	{
+		return false;
+	}
+	Pos = End + 1;
+
+	int32 Month = FCString::Strtoi(Pos, &End, 10);
+	if(End <= Pos || *End != '/')
+	{
+		return false;
+	}
+	Pos = End + 1;
+
+	int32 Day = FCString::Strtoi(Pos, &End, 10);
+	if(End <= Pos || *End != 0)
+	{
+		return false;
+	}
+
+	// Parse the time
+	Pos = *Time;
+
+	int Hour = FCString::Strtoi(Pos, &End, 10);
+	if(End <= Pos || *End != ':')
+	{
+		return false;
+	}
+	Pos = End  + 1;
+
+	int Minute = FCString::Strtoi(Pos, &End, 10);
+	if(End <= Pos || *End != ':')
+	{
+		return false;
+	}
+	Pos = End + 1;
+
+	int Second = FCString::Strtoi(Pos, &End, 10);
+	if(End <= Pos || *End != 0)
+	{
+		return false;
+	}
+
+	OutDate = FDateTime(Year, Month, Day, Hour, Minute, Second);
+	return true;
+}
diff --git a/Engine/Source/Programs/SlateUGS/Private/UGSCore/Perforce.h b/Engine/Source/Programs/SlateUGS/Private/UGSCore/Perforce.h
new file mode 100644
--- /dev/null
+++ b/Engine/Source/Programs/SlateUGS/Private/UGSCore/Perforce.h
@@ -0,0 +1,271 @@
+// Copyright 1998-2017 Epic Games, Inc. All Rights Reserved.
+
+#pragma once
+
+#include "CoreMinimal.h"
+#include "Misc/Guid.h"
+#include "Misc/EnumClassFlags.h"
+#include "Templates/Tuple.h"
+
+struct FPerforceChangeSummary
+{
+	int Number;
+	FDateTime Date;
+	FString User;
+	FString Client;
+	FString Description;
+
+	FPerforceChangeSummary();
+};
+
+struct FPerforceFileChangeSummary
+{
+	int Revision;
+	int ChangeNumber;
+	FString Action;
+	FDateTime Date;
+	FString User;
+	FString Client;
+	FString Type;
+	FString Description;
+
+	FPerforceFileChangeSummary();
+};
+
+struct FPerforceClientRecord
+{
+	FString Name;
+	FString Owner;
+	FString Host;
+	FString Root;
+
+	FPerforceClientRecord(const TMap<FString, FString>& Tags);
+};
+
+struct FPerforceDescribeFileRecord
+{
+	FString DepotFile;
+	FString Action;
+	FString Type;
+	int Revision;
+	int FileSize;
+	FString Digest;
+
+	FPerforceDescribeFileRecord();
+};
+
+struct FPerforceDescribeRecord
+{
+	int ChangeNumber;
+	FString User;
+	FString Client;
+	int64 Time;
+	FString Description;
+	FString Status;
+	FString ChangeType;
+	FString Path;
+	TArray<FPerforceDescribeFileRecord> Files;
+
+	FPerforceDescribeRecord(const TMap<FString, FString>& Tags);
+};
+
+struct FPerforceInfoRecord
+{
+	FString UserName;
+	FString HostName;
+	FString ClientAddress;
+	FString ServerAddress;
+	FTimespan ServerTimeZone;
+
+	FPerforceInfoRecord(const TMap<FString, FString>& Tags);
+};
+
+struct FPerforceFileRecord
+{
+	FString DepotPath;
+	FString ClientPath;
+	FString Path;
+	FString Action;
+	int Revision;
+	bool IsMapped;
+	bool Unmap;
+
+	FPerforceFileRecord(const TMap<FString, FString>& Tags);
+};
+
+struct FPerforceTagRecordParser
+{
+	TFunction<void(const TMap<FString, FString>&)> OutputRecord;
+	TMap<FString, FString> Tags;
+
+	FPerforceTagRecordParser(TFunction<void(const TMap<FString, FString>&)> InOutputRecord);
+	~FPerforceTagRecordParser();
+
+	void OutputLine(const FString& Line);
+	void Flush();
+};
+
+struct FPerforceWhereRecord
+{
+	FString LocalPath;
+	FString ClientPath;
+	FString DepotPath;
+};
+
+struct FPerforceSyncOptions
+{
+	int NumRetries;
+	int NumThreads;
+	int TcpBufferSize;
+
+	FPerforceSyncOptions();
+};
+
+struct FPerforceSpec
+{
+	TArray<TTuple<FString, FString>> Sections;
+
+	/**
+	 * Gets the current value of a field with the given name 
+	 *
+	 * @param Name Name of the field to search for
+	 * @return The value of the field, or null if it does not exist.
+	 */
+	const TCHAR* GetField(const TCHAR* Name) const;
+
+	/**
+	 * Sets the value of an existing field, or adds a new one with the given name
+	 *
+	 * @param Name of the field to set
+	 * @param New value of the field
+	 */
+	void SetField(const TCHAR* Name, const TCHAR* Value);
+
+	/**
+	 * Parses a spec (clientspec, branchspec, changespec) from an array of lines
+	 *
+	 * @param Lines Text split into separate lines
+	 * @return Array of section names and values
+	 */
+	static bool TryParse(const TArray<FString>& Lines, TSharedPtr<FPerforceSpec>& OutSpec, FOutputDevice& Log);
+
+	/**
+	 * Formats a P4 specification as a block of text
+	 *
+	 * @return Spec as a single string
+	 */
+	FString ToString() const;
+};
+
+enum class EPerforceOutputChannel
+{
+	Unknown,
+	Text,
+	Info,
+	TaggedInfo,
+	Warning,
+	Error,
+	Exit
+};
+
+struct FPerforceOutputLine
+{
+	EPerforceOutputChannel Channel;
+	FString Text;
+
+	FPerforceOutputLine();
+	FPerforceOutputLine(EPerforceOutputChannel InChannel, const FString& InText);
+};
+
+class FPerforceConnection
+{
+public:
+	enum class ECommandOptions
+	{
+		None = 0x0,
+		NoClient = 0x1,
+		NoFailOnErrors = 0x2,
+		NoChannels = 0x4,
+		IgnoreFilesUpToDateError = 0x8,
+		IgnoreNoSuchFilesError = 0x10,
+		IgnoreFilesNotOpenedOnThisClientError = 0x20,
+		IgnoreExitCode = 0x40,
+		IgnoreFilesNotInClientViewError = 0x80,
+	};
+
+	const FString ServerAndPort;
+	const FString UserName;
+	const FString ClientName;
+
+	FPerforceConnection(const TCHAR* InUserName, const TCHAR* InClientName, const TCHAR* InServerAndPort);
+	TSharedRef<FPerforceConnection> OpenClient(const FString& NewClientName) const;
+
+	bool Info(TSharedPtr<FPerforceInfoRecord>& OutInfo, FEvent* AbortEvent, FOutputDevice& Log) const;
+	bool GetSetting(const FString& Name, FString& Value, FEvent* AbortEvent, FOutputDevice& Log) const;
+	bool FindClients(TArray<FPerforceClientRecord>& Clients, FEvent* AbortEvent, FOutputDevice& Log) const;
+	bool FindClients(const FString& ForUserName, TArray<FString>& ClientNames, FEvent* AbortEvent, FOutputDevice& Log) const;
+	bool TryGetClientSpec(const FString& ClientName, TSharedPtr<FPerforceSpec>& OutSpec, FEvent* AbortEvent, FOutputDevice& Log) const;
+	bool TryGetStreamSpec(const FString& StreamName, TSharedPtr<FPerforceSpec>& OutSpec, FEvent* AbortEvent, FOutputDevice& Log) const;
+	bool FindFiles(const FString& Filter, TArray<FPerforceFileRecord>& FileRecords, FEvent* AbortEvent, FOutputDevice& Log) const;
+	bool Print(const FString& DepotPath, TArray<FString>& Lines, FEvent* AbortEvent, FOutputDevice& Log) const;
+	bool PrintToFile(const FString& DepotPath, const FString& OutputFileName, FEvent* AbortEvent, FOutputDevice& Log) const;
+	bool FileExists(const FString& Filter, bool& bExists, FEvent* AbortEvent, FOutputDevice& Log) const;
+	bool FindChanges(const FString& Filter, int MaxResults, TArray<FPerforceChangeSummary>& OutChanges, FEvent* AbortEvent, FOutputDevice& Log) const;
+	bool FindChanges(const TArray<FString>& Filters, int MaxResults, TArray<FPerforceChangeSummary>& OutChanges, FEvent* AbortEvent, FOutputDevice& Log) const;
+	bool TryParseChangeSummary(const FString& Line, FPerforceChangeSummary& Change) const;
+	bool FindFileChanges(const FString& FilePath, int MaxResults, TArray<FPerforceFileChangeSummary>& Changes, FEvent* AbortEvent, FOutputDevice& Log) const;
+	bool TryParseFileChangeSummary(const TArray<FString>& Lines, int& LineIdx, FPerforceFileChangeSummary& OutChange) const;
+	bool ConvertToClientPath(const FString& FileName, FString& OutClientFileName, FEvent* AbortEvent, FOutputDevice& Log) const;
+	bool ConvertToDepotPath(const FString& FileName, FString& OutDepotFileName, FEvent* AbortEvent, FOutputDevice& Log) const;
+	bool ConvertToLocalPath(const FString& FileName, FString& OutLocalFileName, FEvent* AbortEvent, FOutputDevice& Log) const;
+	bool FindStreams(const FString& Filter, TArray<FString>& OutStreamNames, FEvent* AbortEvent, FOutputDevice& Log) const;
+	bool HasOpenFiles(FEvent* AbortEvent, FOutputDevice& Log) const;
+	bool SwitchStream(const FString& NewStream, FEvent* AbortEvent, FOutputDevice& Log) const;
+	bool Describe(int ChangeNumber, TSharedPtr<FPerforceDescribeRecord>& OutRecord, FEvent* AbortEvent, FOutputDevice& Log) const;
+	bool Where(const FString& Filter, TSharedPtr<FPerforceWhereRecord>& OutWhereRecord, FEvent* AbortEvent, FOutputDevice& Log) const;
+	bool Have(const FString& Filter, TArray<FPerforceFileRecord>& OutFileRecords, FEvent* AbortEvent, FOutputDevice& Log) const;
+	bool Stat(const FString& Filter, TArray<FPerforceFileRecord>& OutFileRecords, FEvent* AbortEvent, FOutputDevice& Log) const;
+	bool Sync(const FString& Filter, FEvent* AbortEvent, FOutputDevice& Log) const;
+	bool Sync(const TArray<FString>& DepotPaths, int ChangeNumber, TFunction<void(const FPerforceFileRecord&)> SyncOutput, TArray<FString>& OutTamperedFiles, const FPerforceSyncOptions* Options, FEvent* AbortEvent, FOutputDevice& Log) const;
+
+private:
+	static bool FilterSyncOutput(const FPerforceOutputLine& Line, FPerforceTagRecordParser& Parser, TArray<FString>& OutTamperedFiles, FOutputDevice& Log);
+	static void ParseTamperedFile(const FString& Line, TArray<FString>& OutTamperedFiles);
+
+public:
+	bool SyncPreview(const FString& Filter, int ChangeNumber, bool bOnlyFilesInThisChange, TArray<FPerforceFileRecord>& OutFileRecords, FEvent* AbortEvent, FOutputDevice& Log) const;
+	bool ForceSync(const FString& Filter, int ChangeNumber, FEvent* AbortEvent, FOutputDevice& Log) const;
+	bool GetOpenFiles(const FString& Filter, TArray<FPerforceFileRecord>& OutFileRecords, FEvent* AbortEvent, FOutputDevice& Log) const;
+	bool GetUnresolvedFiles(const FString& Filter, TArray<FPerforceFileRecord>& OutFileRecords, FEvent* AbortEvent, FOutputDevice& Log) const;
+	bool AutoResolveFile(const FString& Filter, FEvent* AbortEvent, FOutputDevice& Log) const;
+	bool GetActiveStream(FString& OutStreamName, FEvent* AbortEvent, FOutputDevice& Log) const;
+
+private:
+	bool RunCommand(const FString& CommandLine, TArray<FPerforceFileRecord>& OutFileRecords, ECommandOptions Options, FEvent* AbortEvent, FOutputDevice& Log) const;
+	bool RunCommand(const FString& CommandLine, TArray<FPerforceClientRecord>& OutClientRecords, ECommandOptions Options, FEvent* AbortEvent, FOutputDevice& Log) const;
+	bool RunCommand(const FString& CommandLine, TArray<TMap<FString, FString>>& OutTagRecords, ECommandOptions Options, FEvent* AbortEvent, FOutputDevice& Log) const;
+	bool RunCommand(const FString& CommandLine, ECommandOptions Options, FEvent* AbortEvent, FOutputDevice& Log) const;
+	bool RunCommand(const FString& CommandLine, TArray<FString>& OutLines, ECommandOptions Options, FEvent* AbortEvent, FOutputDevice& Log) const;
+	bool RunCommand(const FString& CommandLine, EPerforceOutputChannel Channel, TArray<FString>& OutLines, ECommandOptions Options, FEvent* AbortEvent, FOutputDevice& Log) const;
+	bool RunCommand(const FString& CommandLine, const TCHAR* Input, TFunction<bool(const FPerforceOutputLine&)> HandleOutput, ECommandOptions Options, FEvent* AbortEvent, FOutputDevice& Log) const;
+	FString GetFullCommandLine(const FString& CommandLine, ECommandOptions Options) const;
+	bool ParseCommandOutput(const FString& Text, TFunction<bool(const FPerforceOutputLine&)> HandleOutput, ECommandOptions Options) const;
+	static bool IsValidTag(const FString& Line, int StartIndex);
+	static bool IgnoreCommandOutput(const FString& Text, ECommandOptions Options);
+
+	static const uint8* ReadBinaryField(const uint8* Pos, const uint8* End, FString* OutField);
+	static const uint8* ReadBinaryRecord(const uint8* Pos, const uint8* End, TMap<FString, FString>* OutRecord);
+
+	bool RunCommandWithBinaryOutput(const FString& CommandLine, TArray<TMap<FString, FString>>& OutRecords, ECommandOptions Options, FEvent* AbortEvent, FOutputDevice& Log) const;
+	bool RunCommandWithBinaryOutput(const FString& CommandLine, TFunction<void(const TMap<FString, FString>&)> HandleOutput, ECommandOptions Options, FEvent* AbortEvent, FOutputDevice& Log) const;
+};
+
+ENUM_CLASS_FLAGS(FPerforceConnection::ECommandOptions)
+
+//// FPerforceUtils ////
+
+struct FPerforceUtils
+{
+	static FString GetClientOrDepotDirectoryName(const TCHAR* ClientFile);
+	static bool TryParseDateTime(const FString& Date, const FString& Time, FDateTime& OutDate);
+};
diff --git a/Engine/Source/Programs/SlateUGS/Private/UGSCore/Telemetry.cpp b/Engine/Source/Programs/SlateUGS/Private/UGSCore/Telemetry.cpp
new file mode 100644
--- /dev/null
+++ b/Engine/Source/Programs/SlateUGS/Private/UGSCore/Telemetry.cpp
@@ -0,0 +1,52 @@
+// Copyright 1998-2017 Epic Games, Inc. All Rights Reserved.
+
+#include "Telemetry.h"
+
+//// FTelemetryStopwatch ////
+
+FTelemetryStopwatch::FTelemetryStopwatch(const FString& InAction, const FString& InProject)
+	: Action(InAction)
+	, Project(InProject)
+	, StartTime(FDateTime::UtcNow())
+	, EndTime(FDateTime(0))
+{
+}
+
+FTelemetryStopwatch::~FTelemetryStopwatch()
+{
+	if(Result.Len() == 0)
+	{
+		Stop(TEXT("Aborted"));
+	}
+//	TelemetryWriter.Enqueue(Action, Result, Project, StartTime, (float)Elapsed.TotalSeconds);
+}
+
+FTimespan FTelemetryStopwatch::Stop(const FString& InResult)
+{
+	EndTime = FDateTime::UtcNow();
+	Result = InResult;
+	return EndTime - StartTime;
+}
+
+FTimespan FTelemetryStopwatch::GetElapsed() const
+{
+	if(Result.Len() == 0)
+	{
+		return FDateTime::UtcNow() - StartTime;
+	}
+	else
+	{
+		return EndTime - StartTime;
+	}
+}
+
+//// FTelemetryWriter ////
+
+FTelemetryWriter::FTelemetryWriter(const FString& InSqlConnectionString, const FString& InLogFileName)
+{
+}
+
+FTelemetryWriter::~FTelemetryWriter()
+{
+}
+
diff --git a/Engine/Source/Programs/SlateUGS/Private/UGSCore/Telemetry.h b/Engine/Source/Programs/SlateUGS/Private/UGSCore/Telemetry.h
new file mode 100644
--- /dev/null
+++ b/Engine/Source/Programs/SlateUGS/Private/UGSCore/Telemetry.h
@@ -0,0 +1,30 @@
+// Copyright 1998-2017 Epic Games, Inc. All Rights Reserved.
+
+#pragma once
+
+#include "CoreMinimal.h"
+#include "Misc/Timespan.h"
+
+class FTelemetryStopwatch
+{
+public:
+	FTelemetryStopwatch(const FString& InAction, const FString& InProject);
+	~FTelemetryStopwatch();
+
+	FTimespan Stop(const FString& InResult);
+	FTimespan GetElapsed() const;
+
+private:
+	const FString Action;
+	const FString Project;
+	const FDateTime StartTime;
+	FString Result;
+	FDateTime EndTime;
+};
+
+class FTelemetryWriter
+{
+public:
+	FTelemetryWriter(const FString& InSqlConnectionString, const FString& InLogFileName);
+	~FTelemetryWriter();
+};
\ No newline at end of file
diff --git a/Engine/Source/Programs/SlateUGS/Private/UGSCore/UserSettings.cpp b/Engine/Source/Programs/SlateUGS/Private/UGSCore/UserSettings.cpp
new file mode 100644
--- /dev/null
+++ b/Engine/Source/Programs/SlateUGS/Private/UGSCore/UserSettings.cpp
@@ -0,0 +1,611 @@
+// Copyright 1998-2017 Epic Games, Inc. All Rights Reserved.
+
+#include "UserSettings.h"
+#include "HAL/FileManager.h"
+#include "Algo/Transform.h"
+#include "Utility.h"
+
+//// EBuildConfig ////
+
+FString ToString(EBuildConfig BuildConfig)
+{
+	switch(BuildConfig)
+	{
+	case EBuildConfig::Debug:
+		return TEXT("Debug");
+	case EBuildConfig::DebugGame:
+		return TEXT("DebugGame");
+	case EBuildConfig::Development:
+		return TEXT("Development");
+	default:
+		check(false);
+		return FString();
+	}
+}
+
+bool TryParse(const TCHAR* Text, EBuildConfig& OutBuildConfig)
+{
+	if (FCString::Stricmp(Text, TEXT("Debug")) == 0)
+	{
+		OutBuildConfig = EBuildConfig::Debug;
+		return true;
+	}
+	else if (FCString::Stricmp(Text, TEXT("DebugGame")) == 0)
+	{
+		OutBuildConfig = EBuildConfig::DebugGame;
+		return true;
+	}
+	else if (FCString::Stricmp(Text, TEXT("Development")) == 0)
+	{
+		OutBuildConfig = EBuildConfig::Development;
+		return true;
+	}
+	return false;
+}
+
+//// ETabLabels ////
+
+FString ToString(ETabLabels TabLabels)
+{
+	switch(TabLabels)
+	{
+	case ETabLabels::Stream:
+		return TEXT("Stream");
+	case ETabLabels::WorkspaceName:
+		return TEXT("WorkspaceName");
+	case ETabLabels::WorkspaceRoot:
+		return TEXT("WorkspaceRoot");
+	case ETabLabels::ProjectFile:
+		return TEXT("ProjectFile");
+	default:
+		check(false);
+		return FString();
+	}
+}
+
+bool TryParse(const TCHAR* Text, ETabLabels& OutTabLabels)
+{
+	if (FCString::Stricmp(Text, TEXT("Stream")) == 0)
+	{
+		OutTabLabels = ETabLabels::Stream;
+		return true;
+	}
+	else if (FCString::Stricmp(Text, TEXT("WorkspaceName")) == 0)
+	{
+		OutTabLabels = ETabLabels::WorkspaceName;
+		return true;
+	}
+	else if (FCString::Stricmp(Text, TEXT("WorkspaceRoot")) == 0)
+	{
+		OutTabLabels = ETabLabels::WorkspaceRoot;
+		return true;
+	}
+	else if (FCString::Stricmp(Text, TEXT("ProjectFile")) == 0)
+	{
+		OutTabLabels = ETabLabels::ProjectFile;
+		return true;
+	}
+	return false;
+}
+
+//// ELatestChangeType ////
+
+FString ToString(ELatestChangeType LatestChangeType)
+{
+	switch(LatestChangeType)
+	{
+	case ELatestChangeType::Any:
+		return TEXT("Any");
+	case ELatestChangeType::Good:
+		return TEXT("Good");
+	case ELatestChangeType::Starred:
+		return TEXT("Starred");
+	default:
+		check(false);
+		return FString();
+	}
+}
+
+bool TryParse(const TCHAR* Text, ELatestChangeType& OutLatestChangeType)
+{
+	if(FCString::Stricmp(Text, TEXT("Any")) == 0)
+	{
+		OutLatestChangeType = ELatestChangeType::Any;
+		return true;
+	}
+	else if(FCString::Stricmp(Text, TEXT("Good")) == 0)
+	{
+		OutLatestChangeType = ELatestChangeType::Good;
+		return true;
+	}
+	else if(FCString::Stricmp(Text, TEXT("Starred")) == 0)
+	{
+		OutLatestChangeType = ELatestChangeType::Starred;
+		return true;
+	}
+	return false;
+}
+
+//// FUserSettings ////
+
+FUserSettings::FUserSettings(const FString& InFileName)
+	: FileName(InFileName)
+{
+	if(IFileManager::Get().FileExists(*FileName))
+	{
+		ConfigFile.Load(*FileName);
+	}
+
+	// General settings
+	bBuildAfterSync = (ConfigFile.GetValue(TEXT("General.BuildAfterSync"), 1) != 0);
+	bRunAfterSync = (ConfigFile.GetValue(TEXT("General.RunAfterSync"), 1) != 0);
+	bSyncPrecompiledEditor = (ConfigFile.GetValue(TEXT("General.SyncPrecompiledEditor"), 0) != 0);
+	bOpenSolutionAfterSync = (ConfigFile.GetValue(TEXT("General.OpenSolutionAfterSync"), 0) != 0);
+	bShowLogWindow = (ConfigFile.GetValue(TEXT("General.ShowLogWindow"), false));
+	bAutoResolveConflicts = (ConfigFile.GetValue(TEXT("General.AutoResolveConflicts"), 1) != 0);
+	bUseIncrementalBuilds = ConfigFile.GetValue(TEXT("General.IncrementalBuilds"), true);
+	bShowLocalTimes = ConfigFile.GetValue(TEXT("General.ShowLocalTimes"), false);
+	bShowAllStreams = ConfigFile.GetValue(TEXT("General.ShowAllStreams"), false);
+	bKeepInTray = ConfigFile.GetValue(TEXT("General.KeepInTray"), true);
+	FilterIndex = ConfigFile.GetValue(TEXT("General.FilterIndex"), 0);
+	LastProjectFileName = ConfigFile.GetValue(TEXT("General.LastProjectFileName"), TEXT(""));
+
+	ConfigFile.TryGetValues(TEXT("General.OpenProjectFileNames"), OpenProjectFileNames);
+	if(LastProjectFileName.Len() > 0)
+	{
+		for(int Idx = 0;;Idx++)
+		{
+			if(Idx == OpenProjectFileNames.Num())
+			{
+				OpenProjectFileNames.Add(LastProjectFileName);
+				break;
+			}
+			else if(OpenProjectFileNames[Idx] == LastProjectFileName)
+			{
+				break;
+			}
+		}
+	}
+
+	ConfigFile.TryGetValues(TEXT("General.OtherProjectFileNames"), OtherProjectFileNames);
+	ConfigFile.TryGetValues(TEXT("General.SyncFilter"), SyncView);
+	ConfigFile.TryGetValues(TEXT("General.SyncExcludedCategories"), SyncExcludedCategories);
+	if(TryParse(ConfigFile.GetValue(TEXT("General.SyncType"), TEXT("")), SyncType))
+	{
+		SyncType = ELatestChangeType::Good;
+	}
+
+	// Build configuration
+	const TCHAR* CompiledEditorBuildConfigName = ConfigFile.GetValue(TEXT("General.BuildConfig"), TEXT(""));
+	if(!TryParse(CompiledEditorBuildConfigName, CompiledEditorBuildConfig))
+	{
+		CompiledEditorBuildConfig = EBuildConfig::DebugGame;
+	}
+
+	// Tab names
+	const TCHAR* TabNamesValue = ConfigFile.GetValue(TEXT("General.TabNames"), TEXT(""));
+	if(!TryParse(TabNamesValue, TabLabels))
+	{
+		TabLabels = ETabLabels::ProjectFile;
+	}
+
+	// Editor arguments
+	TArray<FString> Arguments;
+	if(!ConfigFile.TryGetValues(TEXT("General.EditorArguments"), Arguments))
+	{
+		Arguments.Add(TEXT("0:-log"));
+		Arguments.Add(TEXT("0:-fastload"));
+	}
+	for(const FString& Argument : Arguments)
+	{
+		if(Argument.StartsWith(TEXT("0:")))
+		{
+			EditorArguments.Add(TTuple<FString,bool>(Argument.Mid(2), false));
+		}
+		else if(Argument.StartsWith(TEXT("1:")))
+		{
+			EditorArguments.Add(TTuple<FString,bool>(Argument.Mid(2), true));
+		}
+		else
+		{
+			EditorArguments.Add(TTuple<FString,bool>(Argument, true));
+		}
+	}
+
+	// Window settings
+	TSharedPtr<const FCustomConfigSection> WindowSection = ConfigFile.FindSection(TEXT("Window"));
+	if(WindowSection.IsValid())
+	{
+		bHasWindowSettings = true;
+
+		int X = WindowSection->GetValue(TEXT("X"), -1);
+		int Y = WindowSection->GetValue(TEXT("Y"), -1);
+		int Width = WindowSection->GetValue(TEXT("Width"), -1);
+		int Height = WindowSection->GetValue(TEXT("Height"), -1);
+//		WindowRectangle = new Rectangle(X, Y, Width, Height);
+
+		FCustomConfigObject ColumnWidthObject(WindowSection->GetValue(TEXT("ColumnWidths"), TEXT("")));
+		for(const TTuple<FString, FString>& ColumnWidthPair : ColumnWidthObject.Pairs)
+		{
+//			int Value;
+//			if(int.TryParse(ColumnWidthPair.Value, out Value))
+//			{
+//				ColumnWidths[ColumnWidthPair.Key] = Value;
+//			}
+		}
+
+		bWindowVisible = WindowSection->GetValue(TEXT("Visible"), true);
+	}
+
+	// Schedule settings
+	bScheduleEnabled = ConfigFile.GetValue(TEXT("Schedule.Enabled"), false);
+	if(!FTimespan::Parse(ConfigFile.GetValue(TEXT("Schedule.Time"), TEXT("")), ScheduleTime))
+	{
+		ScheduleTime = FTimespan(6, 0, 0);
+	}
+	if(!TryParse(ConfigFile.GetValue(TEXT("Schedule.Change"), TEXT("")), ScheduleChange))
+	{
+		ScheduleChange = ELatestChangeType::Good;
+	}
+
+	// Perforce settings
+	SyncOptions.NumRetries = ConfigFile.GetValue(TEXT("Perforce.NumRetries"), 0);
+	SyncOptions.NumThreads = ConfigFile.GetValue(TEXT("Perforce.NumThreads"), 0);
+	SyncOptions.TcpBufferSize = ConfigFile.GetValue(TEXT("Perforce.TcpBufferSize"), 0);
+}
+
+FString Trim(const FString& Text, const TCHAR Character)
+{
+	FString Result = Text;
+	while(Result.StartsWith(TEXT("/")))
+	{
+		Result = Result.Mid(1);
+	}
+	while(Result.EndsWith(TEXT("/")))
+	{
+		Result = Result.Mid(0, Result.Len() - 1);
+	}
+	return Result;
+}
+
+TSharedRef<FUserWorkspaceSettings> FUserSettings::FindOrAddWorkspace(const TCHAR* ClientBranchPath)
+{
+	// Update the current workspace
+	FString CurrentWorkspaceKey = Trim(ClientBranchPath, TEXT('/'));
+	// TODO: Trim('/')?
+
+	TSharedRef<FUserWorkspaceSettings>* CurrentWorkspaceEntry = WorkspaceKeyToSettings.Find(CurrentWorkspaceKey);
+	if(CurrentWorkspaceEntry == nullptr)
+	{
+		// Create a new workspace settings object
+		TSharedRef<FUserWorkspaceSettings> CurrentWorkspace = MakeShared<FUserWorkspaceSettings>();
+		CurrentWorkspaceEntry = &(WorkspaceKeyToSettings.Add(CurrentWorkspaceKey, CurrentWorkspace));
+
+		// Read the workspace settings
+		TSharedPtr<FCustomConfigSection> WorkspaceSection = ConfigFile.FindSection(*CurrentWorkspaceKey);
+		if(!WorkspaceSection.IsValid())
+		{
+			FString LegacyBranchAndClientKey = Trim(ClientBranchPath, TEXT('/'));
+
+			int SlashIdx;
+			if(LegacyBranchAndClientKey.FindChar(TEXT('/'), SlashIdx))
+			{
+				LegacyBranchAndClientKey = LegacyBranchAndClientKey.Mid(0, SlashIdx) + "$" + LegacyBranchAndClientKey.Mid(SlashIdx + 1);
+			}
+
+			FString CurrentSync;
+			if(ConfigFile.TryGetValue(*(FString(TEXT("Clients.")) + LegacyBranchAndClientKey), CurrentSync))
+			{
+				int AtIdx;
+				if(CurrentSync.FindLastChar(TEXT('@'), AtIdx))
+				{
+					int ChangeNumber;
+					if(FUtility::TryParse(*CurrentSync.Mid(AtIdx + 1), ChangeNumber))
+					{
+						CurrentWorkspace->CurrentProjectIdentifier = CurrentSync.Mid(0, AtIdx);
+						CurrentWorkspace->CurrentChangeNumber = ChangeNumber;
+					}
+				}
+			}
+
+			FString LastUpdateResultText;
+			if(ConfigFile.TryGetValue(*(FString(TEXT("Clients.")) + LegacyBranchAndClientKey + TEXT("$LastUpdate")), LastUpdateResultText))
+			{
+				int ColonIdx;
+				if(LastUpdateResultText.FindLastChar(TEXT(':'), ColonIdx))
+				{
+					int ChangeNumber;
+					if(FUtility::TryParse(*LastUpdateResultText.Mid(0, ColonIdx), ChangeNumber))
+					{
+						EWorkspaceUpdateResult Result;
+						if(TryParse(*LastUpdateResultText.Mid(ColonIdx + 1), Result))
+						{
+							CurrentWorkspace->LastSyncChangeNumber = ChangeNumber;
+							CurrentWorkspace->LastSyncResult = Result;
+						}
+					}
+				}
+			}
+
+			CurrentWorkspace->SyncView.Empty();
+			CurrentWorkspace->SyncExcludedCategories.Empty();
+		}
+		else
+		{
+			CurrentWorkspace->CurrentProjectIdentifier = WorkspaceSection->GetValue(TEXT("CurrentProjectPath"), TEXT(""));
+			CurrentWorkspace->CurrentChangeNumber = WorkspaceSection->GetValue(TEXT("CurrentChangeNumber"), -1);
+			WorkspaceSection->TryGetValues(TEXT("AdditionalChangeNumbers"), CurrentWorkspace->AdditionalChangeNumbers);
+			if(!TryParse(WorkspaceSection->GetValue(TEXT("LastSyncResult"), TEXT("")), CurrentWorkspace->LastSyncResult))
+			{
+				CurrentWorkspace->LastSyncResult = EWorkspaceUpdateResult::Canceled;
+			}
+			CurrentWorkspace->LastSyncResultMessage = UnescapeText(WorkspaceSection->GetValue(TEXT("LastSyncResultMessage"), TEXT("")));
+			CurrentWorkspace->LastSyncChangeNumber = WorkspaceSection->GetValue(TEXT("LastSyncChangeNumber"), -1);
+
+			FDateTime LastSyncTime;
+			if(FDateTime::Parse(WorkspaceSection->GetValue(TEXT("LastSyncTime"), TEXT("")), LastSyncTime))
+			{
+				CurrentWorkspace->LastSyncTime = LastSyncTime;
+			}
+
+			CurrentWorkspace->LastSyncDurationSeconds = WorkspaceSection->GetValue(TEXT("LastSyncDuration"), 0);
+			CurrentWorkspace->LastBuiltChangeNumber = WorkspaceSection->GetValue(TEXT("LastBuiltChangeNumber"), 0);
+			WorkspaceSection->TryGetValues(TEXT("ExpandedArchiveName"), CurrentWorkspace->ExpandedArchiveTypes);
+
+			WorkspaceSection->TryGetValues(TEXT("SyncFilter"), CurrentWorkspace->SyncView);
+			WorkspaceSection->TryGetValues(TEXT("SyncExcludedCategories"), CurrentWorkspace->SyncExcludedCategories);
+		}
+	}
+	return *CurrentWorkspaceEntry;
+}
+
+TSharedRef<FUserProjectSettings> FUserSettings::FindOrAddProject(const TCHAR* ClientProjectFileName)
+{
+	// Read the project settings
+	TSharedRef<FUserProjectSettings>* CurrentProject = ProjectKeyToSettings.Find(ClientProjectFileName);
+	if(CurrentProject == nullptr)
+	{
+		TSharedRef<FUserProjectSettings> NewCurrentProject = MakeShared<FUserProjectSettings>();
+		CurrentProject = &(ProjectKeyToSettings.Add(ClientProjectFileName, NewCurrentProject));
+	
+		TArray<FString> BuildStepStrings;
+		TSharedRef<FCustomConfigSection> ProjectSection = ConfigFile.FindOrAddSection(ClientProjectFileName);
+		ProjectSection->TryGetValues(TEXT("BuildStep"), BuildStepStrings);
+		Algo::Transform(BuildStepStrings, NewCurrentProject->BuildSteps, [](const FString& Text) -> FCustomConfigObject { return FCustomConfigObject(*Text); });
+	}
+	return *CurrentProject;
+}
+
+void FUserSettings::Save()
+{
+	// General settings
+	TSharedRef<FCustomConfigSection> GeneralSection = ConfigFile.FindOrAddSection(TEXT("General"));
+	GeneralSection->Clear();
+	GeneralSection->SetValue(TEXT("BuildAfterSync"), bBuildAfterSync);
+	GeneralSection->SetValue(TEXT("RunAfterSync"), bRunAfterSync);
+	GeneralSection->SetValue(TEXT("SyncPrecompiledEditor"), bSyncPrecompiledEditor);
+	GeneralSection->SetValue(TEXT("OpenSolutionAfterSync"), bOpenSolutionAfterSync);
+	GeneralSection->SetValue(TEXT("ShowLogWindow"), bShowLogWindow);
+	GeneralSection->SetValue(TEXT("AutoResolveConflicts"), bAutoResolveConflicts);
+	GeneralSection->SetValue(TEXT("IncrementalBuilds"), bUseIncrementalBuilds);
+	GeneralSection->SetValue(TEXT("ShowLocalTimes"), bShowLocalTimes);
+	GeneralSection->SetValue(TEXT("ShowAllStreams"), bShowAllStreams);
+	GeneralSection->SetValue(TEXT("LastProjectFileName"), *LastProjectFileName);
+	GeneralSection->SetValues(TEXT("OpenProjectFileNames"), OpenProjectFileNames);
+	GeneralSection->SetValue(TEXT("KeepInTray"), bKeepInTray);
+	GeneralSection->SetValue(TEXT("FilterIndex"), FilterIndex);
+	GeneralSection->SetValues(TEXT("OtherProjectFileNames"), OtherProjectFileNames);
+	GeneralSection->SetValues(TEXT("SyncFilter"), SyncView);
+	GeneralSection->SetValues(TEXT("SyncExcludedCategories"), SyncExcludedCategories);
+	GeneralSection->SetValue(TEXT("SyncType"), *ToString(SyncType));
+
+	// Build configuration
+	GeneralSection->SetValue(TEXT("BuildConfig"), *ToString(CompiledEditorBuildConfig));
+
+	// Tab names
+	GeneralSection->SetValue(TEXT("TabNames"), *ToString(TabLabels));
+
+	// Editor arguments
+	TArray<FString> EditorArgumentList;
+	for(const TTuple<FString, bool>& EditorArgument : EditorArguments)
+	{
+		EditorArgumentList.Add(FString::Printf(TEXT("%d:%s"), EditorArgument.Value? 1 : 0, *EditorArgument.Key));
+	}
+	GeneralSection->SetValues(TEXT("EditorArguments"), EditorArgumentList);
+
+	// Schedule settings
+	TSharedRef<FCustomConfigSection> ScheduleSection = ConfigFile.FindOrAddSection(TEXT("Schedule"));
+	ScheduleSection->Clear();
+	ScheduleSection->SetValue(TEXT("Enabled"), bScheduleEnabled);
+	ScheduleSection->SetValue(TEXT("Time"), *ScheduleTime.ToString());
+	ScheduleSection->SetValue(TEXT("Change"), *ToString(ScheduleChange));
+
+	// Window settings
+	if(bHasWindowSettings)
+	{
+		TSharedRef<FCustomConfigSection> WindowSection = ConfigFile.FindOrAddSection(TEXT("Window"));
+		WindowSection->Clear();
+//		WindowSection.SetValue(TEXT("X"), WindowRectangle.X);
+//		WindowSection.SetValue(TEXT("Y"), WindowRectangle.Y);
+//		WindowSection.SetValue(TEXT("Width"), WindowRectangle.Width);
+//		WindowSection.SetValue(TEXT("Height"), WindowRectangle.Height);
+
+		FCustomConfigObject ColumnWidthsObject;
+		for(const TTuple<FString, int>& ColumnWidthPair : ColumnWidths)
+		{
+			ColumnWidthsObject.SetValue(*ColumnWidthPair.Key, ColumnWidthPair.Value);
+		}
+		WindowSection->SetValue(TEXT("ColumnWidths"), *ColumnWidthsObject.ToString());
+
+		WindowSection->SetValue(TEXT("Visible"), bWindowVisible);
+	}
+
+	// Current workspace settings
+	for(TTuple<FString, TSharedRef<FUserWorkspaceSettings>>& Pair : WorkspaceKeyToSettings)
+	{
+		const FString& CurrentWorkspaceKey = Pair.Key;
+		FUserWorkspaceSettings& CurrentWorkspace = Pair.Value.Get();
+
+		TSharedRef<FCustomConfigSection> WorkspaceSection = ConfigFile.FindOrAddSection(*CurrentWorkspaceKey);
+		WorkspaceSection->Clear();
+		WorkspaceSection->SetValue(TEXT("CurrentProjectPath"), *CurrentWorkspace.CurrentProjectIdentifier);
+		WorkspaceSection->SetValue(TEXT("CurrentChangeNumber"), CurrentWorkspace.CurrentChangeNumber);
+		WorkspaceSection->SetValues(TEXT("AdditionalChangeNumbers"), CurrentWorkspace.AdditionalChangeNumbers);
+		WorkspaceSection->SetValue(TEXT("LastSyncResult"), *ToString(CurrentWorkspace.LastSyncResult));
+		WorkspaceSection->SetValue(TEXT("LastSyncResultMessage"), *EscapeText(CurrentWorkspace.LastSyncResultMessage));
+		WorkspaceSection->SetValue(TEXT("LastSyncChangeNumber"), CurrentWorkspace.LastSyncChangeNumber);
+		if(CurrentWorkspace.LastSyncTime.IsSet())
+		{
+			WorkspaceSection->SetValue(TEXT("LastSyncTime"), *CurrentWorkspace.LastSyncTime.GetValue().ToString());
+		}
+		if(CurrentWorkspace.LastSyncDurationSeconds > 0)
+		{
+			WorkspaceSection->SetValue(TEXT("LastSyncDuration"), CurrentWorkspace.LastSyncDurationSeconds);
+		}
+		WorkspaceSection->SetValue(TEXT("LastBuiltChangeNumber"), CurrentWorkspace.LastBuiltChangeNumber);
+		WorkspaceSection->SetValues(TEXT("ExpandedArchiveName"), CurrentWorkspace.ExpandedArchiveTypes);
+		WorkspaceSection->SetValues(TEXT("SyncFilter"), CurrentWorkspace.SyncView);
+		WorkspaceSection->SetValues(TEXT("SyncExcludedCategories"), CurrentWorkspace.SyncExcludedCategories);
+	}
+
+	// Current project settings
+	for(TTuple<FString, TSharedRef<FUserProjectSettings>>& Pair : ProjectKeyToSettings)
+	{
+		const FString& CurrentProjectKey = Pair.Key;
+		FUserProjectSettings& CurrentProject = Pair.Value.Get();
+
+		TSharedRef<FCustomConfigSection> ProjectSection = ConfigFile.FindOrAddSection(*CurrentProjectKey);
+		ProjectSection->Clear();
+
+		TArray<FString> BuildStepStrings;
+		Algo::Transform(CurrentProject.BuildSteps, BuildStepStrings, [](const FCustomConfigObject& Step){ return Step.ToString(); });
+		ProjectSection->SetValues(TEXT("BuildStep"), BuildStepStrings);
+	}
+
+	// Perforce settings
+	TSharedRef<FCustomConfigSection> PerforceSection = ConfigFile.FindOrAddSection(TEXT("Perforce"));
+	PerforceSection->Clear();
+	if(SyncOptions.NumRetries > 0)
+	{
+		PerforceSection->SetValue(TEXT("NumRetries"), SyncOptions.NumRetries);
+	}
+	if(SyncOptions.NumThreads > 0)
+	{
+		PerforceSection->SetValue(TEXT("NumThreads"), SyncOptions.NumThreads);
+	}
+	if(SyncOptions.TcpBufferSize > 0)
+	{
+		PerforceSection->SetValue(TEXT("TcpBufferSize"), SyncOptions.TcpBufferSize);
+	}
+
+	// Save the file
+	ConfigFile.Save(*FileName);
+}
+
+TArray<FString> FUserSettings::GetCombinedSyncFilter(const TMap<FGuid, FWorkspaceSyncCategory>& UniqueIdToFilter, const TArray<FString>& GlobalView, const TArray<FGuid>& GlobalExcludedCategories, const TArray<FString>& WorkspaceView, const TArray<FGuid>& WorkspaceExcludedCategories)
+{
+	TArray<FString> Lines;
+
+	// Add the custom global view
+	for(const FString& GlobalViewLine : GlobalView)
+	{
+		FString TrimLine = GlobalViewLine.TrimStartAndEnd();
+		if(TrimLine.Len() > 0 && !TrimLine.StartsWith(TEXT(";")))
+		{
+			Lines.Add(MoveTemp(TrimLine));
+		}
+	}
+
+	// Add the custom workspace view
+	for(const FString& WorkspaceViewLine : GlobalView)
+	{
+		FString TrimLine = WorkspaceViewLine.TrimStartAndEnd();
+		if(TrimLine.Len() > 0 && !TrimLine.StartsWith(TEXT(";")))
+		{
+			Lines.Add(MoveTemp(TrimLine));
+		}
+	}
+
+	// Add all the excluded categories
+	for(const TTuple<FGuid, FWorkspaceSyncCategory>& Pair : UniqueIdToFilter)
+	{
+		if(Pair.Value.bEnable && (GlobalExcludedCategories.Contains(Pair.Key) || WorkspaceExcludedCategories.Contains(Pair.Key)))
+		{
+			for(const FString& Path : Pair.Value.Paths)
+			{
+				Lines.Add(FString::Printf(TEXT("-%s"), *Path));
+			}
+		}
+	}
+	return Lines;
+}
+
+FString FUserSettings::EscapeText(const FString& Text)
+{
+	FString Result;
+	for(int Idx = 0; Idx < Text.Len(); Idx++)
+	{
+		switch(Text[Idx])
+		{
+			case '\\':
+				Result += TEXT("\\\\");
+				break;
+			case '\t':
+				Result += TEXT("\\t");
+				break;
+			case '\r':
+				Result += TEXT("\\r");
+				break;
+			case '\n':
+				Result += TEXT("\\n");
+				break;
+			case '\'':
+				Result += TEXT("\\\'");
+				break;
+			case '\"':
+				Result += TEXT("\\\"");
+				break;
+			default:
+				Result += Text[Idx];
+				break;
+		}
+	}
+	return Result;
+}
+
+FString FUserSettings::UnescapeText(const FString& Text)
+{
+	FString Result;
+	for(int Idx = 0; Idx < Text.Len(); Idx++)
+	{
+		if(Text[Idx] == '\\' && Idx + 1 < Text.Len())
+		{
+			switch(Text[++Idx])
+			{
+				case 't':
+					Result += TEXT('\t');
+					break;
+				case 'r':
+					Result += TEXT('\r');
+					break;
+				case 'n':
+					Result += TEXT('\n');
+					break;
+				case '\'':
+					Result += TEXT('\'');
+					break;
+				case '\"':
+					Result += TEXT('\"');
+					break;
+				default:
+					Result += Text[Idx];
+					break;
+			}
+		}
+		else
+		{
+			Result += Text[Idx];
+		}
+	}
+	return Result;
+}
diff --git a/Engine/Source/Programs/SlateUGS/Private/UGSCore/UserSettings.h b/Engine/Source/Programs/SlateUGS/Private/UGSCore/UserSettings.h
new file mode 100644
--- /dev/null
+++ b/Engine/Source/Programs/SlateUGS/Private/UGSCore/UserSettings.h
@@ -0,0 +1,128 @@
+// Copyright 1998-2017 Epic Games, Inc. All Rights Reserved.
+
+#pragma once
+
+#include "CoreMinimal.h"
+#include "CustomConfigFile.h"
+#include "Perforce.h"
+#include "Misc/Guid.h"
+#include "Workspace.h"
+
+enum class EBuildConfig
+{
+	Debug,
+	DebugGame,
+	Development,
+};
+
+FString ToString(EBuildConfig BuildConfig);
+bool TryParse(const TCHAR* Text, EBuildConfig& OutBuildConfig);
+
+enum class ETabLabels
+{
+	Stream,
+	WorkspaceName,
+	WorkspaceRoot,
+	ProjectFile,
+};
+
+FString ToString(ETabLabels Lables);
+bool TryParse(const TCHAR* Text, ETabLabels& OutTabLabels);
+
+enum class ELatestChangeType
+{
+	Any,
+	Good,
+	Starred,
+};
+
+FString ToString(ELatestChangeType LatestChangeType);
+bool TryParse(const TCHAR* Text, ELatestChangeType& OutLatestChangeType);
+
+struct FUserWorkspaceSettings
+{
+	// Settings for the currently synced project in this workspace. CurrentChangeNumber is only valid for this workspace if CurrentProjectPath is the current project.
+	FString CurrentProjectIdentifier;
+	int CurrentChangeNumber;
+	TArray<int> AdditionalChangeNumbers;
+
+	// Settings for the last attempted sync. These values are set to persist error messages between runs.
+	int LastSyncChangeNumber;
+	EWorkspaceUpdateResult LastSyncResult;
+	FString LastSyncResultMessage;
+	TOptional<FDateTime> LastSyncTime;
+	int LastSyncDurationSeconds;
+
+	// The last successful build, regardless of whether a failed sync has happened in the meantime. Used to determine whether to force a clean due to entries in the project config file.
+	int LastBuiltChangeNumber;
+
+	// Expanded archives in the workspace
+	TArray<FString> ExpandedArchiveTypes;
+
+	// Workspace specific SyncFilters
+	TArray<FString> SyncView;
+	TArray<FGuid> SyncExcludedCategories;
+};
+
+struct FUserProjectSettings
+{
+	TArray<FCustomConfigObject> BuildSteps;
+};
+
+struct FUserSettings
+{
+	const FString FileName;
+	FCustomConfigFile ConfigFile;
+
+	// General settings
+	bool bBuildAfterSync;
+	bool bRunAfterSync;
+	bool bSyncPrecompiledEditor;
+	bool bOpenSolutionAfterSync;
+	bool bShowLogWindow;
+	bool bAutoResolveConflicts;
+	bool bUseIncrementalBuilds;
+	bool bShowLocalTimes;
+	bool bShowAllStreams;
+	bool bKeepInTray;
+	int FilterIndex;
+	FString LastProjectFileName;
+	TArray<FString> OpenProjectFileNames;
+	TArray<FString> OtherProjectFileNames;
+	TArray<FString> SyncView;
+	TArray<FGuid> SyncExcludedCategories;
+	ELatestChangeType SyncType;
+	EBuildConfig CompiledEditorBuildConfig; // NB: This assumes not using precompiled editor. See CurrentBuildConfig.
+	ETabLabels TabLabels;
+
+	// Window settings
+	bool bHasWindowSettings;
+//	Rectangle WindowRectangle;
+	TMap<FString, int> ColumnWidths;
+	bool bWindowVisible;
+		
+	// Schedule settings
+	bool bScheduleEnabled;
+	FTimespan ScheduleTime;
+	ELatestChangeType ScheduleChange;
+
+	// Run configuration
+	TArray<TTuple<FString, bool>> EditorArguments;
+
+	// Project settings
+	TMap<FString, TSharedRef<FUserWorkspaceSettings>> WorkspaceKeyToSettings;
+	TMap<FString, TSharedRef<FUserProjectSettings>> ProjectKeyToSettings;
+
+	// Perforce settings
+	FPerforceSyncOptions SyncOptions;
+
+	FUserSettings(const FString& InFileName);
+	TSharedRef<FUserWorkspaceSettings> FindOrAddWorkspace(const TCHAR* ClientBranchPath);
+	TSharedRef<FUserProjectSettings> FindOrAddProject(const TCHAR* ClientProjectFileName);
+	void Save();
+	
+	static TArray<FString> GetCombinedSyncFilter(const TMap<FGuid, FWorkspaceSyncCategory>& UniqueIdToFilter, const TArray<FString>& GlobalView, const TArray<FGuid>& GlobalExcludedCategories, const TArray<FString>& WorkspaceView, const TArray<FGuid>& WorkspaceExcludedCategories);
+
+	static FString EscapeText(const FString& Text);
+	static FString UnescapeText(const FString& Text);
+};
diff --git a/Engine/Source/Programs/SlateUGS/Private/UGSCore/Utility.cpp b/Engine/Source/Programs/SlateUGS/Private/UGSCore/Utility.cpp
new file mode 100644
--- /dev/null
+++ b/Engine/Source/Programs/SlateUGS/Private/UGSCore/Utility.cpp
@@ -0,0 +1,242 @@
+// Copyright 1998-2017 Epic Games, Inc. All Rights Reserved.
+
+#include "Utility.h"
+#include "HAL/FileManager.h"
+#include "HAL/PlatformFileManager.h"
+#include "HAL/PlatformProcess.h"
+//#include "ChildProcess.h"
+#include "Misc/Paths.h"
+#include "HAL/Event.h"
+
+bool FUtility::TryParse(const TCHAR* Text, int32& OutValue)
+{
+	TCHAR* TextEnd;
+	OutValue = FCString::Strtoi(Text, &TextEnd, 10);
+	return (TextEnd != Text && *TextEnd == 0);
+}
+
+bool FUtility::TryParse(const TCHAR* Text, int64& OutValue)
+{
+	TCHAR* TextEnd;
+	OutValue = FCString::Strtoi64(Text, &TextEnd, 10);
+	return (TextEnd != Text && *TextEnd == 0);
+}
+
+bool FUtility::IsFileUnderDirectory(const TCHAR* FileName, const TCHAR* DirectoryName)
+{
+	FString FullDirectoryName = FPaths::ConvertRelativePathToFull(DirectoryName);
+	FPaths::MakePlatformFilename(FullDirectoryName);
+
+	if(!FullDirectoryName.EndsWith(FPlatformMisc::GetDefaultPathSeparator()))
+	{
+		FullDirectoryName += FPlatformMisc::GetDefaultPathSeparator();
+	}
+
+	FString FullFileName = FPaths::ConvertRelativePathToFull(FileName);
+	FPaths::MakePlatformFilename(FullFileName);
+
+	return FullFileName.StartsWith(FullDirectoryName);
+}
+
+FString FUtility::GetPathWithCorrectCase(const FString& Path)
+{
+	// Visitor which corrects the case of a filename against any matching item in a directory
+	struct FFixCaseVisitor : public IPlatformFile::FDirectoryVisitor
+	{
+		FString Name;
+
+		FFixCaseVisitor(FString InName) : Name(InName)
+		{
+		}
+
+		virtual bool Visit(const TCHAR* FilenameOrDirectory, bool bIsDirectory) override
+		{
+			if (Name == FilenameOrDirectory)
+			{
+				Name = FilenameOrDirectory;
+				return false;
+			}
+			return true;
+		}
+	};
+
+	FString ParentDir = FPaths::GetPath(Path);
+	if(ParentDir.Len() == 0)
+	{
+		return Path;
+	}
+#if PLATFORM_WINDOWS
+	else if(ParentDir.Len() == 2 && ParentDir[1] == ':')
+	{
+		ParentDir = ParentDir.ToUpper() + FPlatformMisc::GetDefaultPathSeparator();
+	}
+#endif
+	else
+	{
+		ParentDir = GetPathWithCorrectCase(ParentDir) + FPlatformMisc::GetDefaultPathSeparator();
+	}
+
+	FFixCaseVisitor Visitor(ParentDir + FPaths::GetCleanFilename(Path));
+	IFileManager::Get().IterateDirectory(*ParentDir, Visitor);
+	return Visitor.Name;
+}
+
+FString FUtility::FormatUserName(const TCHAR* UserName)
+{
+	FString NormalUserName;
+	for(int Idx = 0; UserName[Idx] != 0; Idx++)
+	{
+		if(Idx == 0 || UserName[Idx - 1] == '.')
+		{
+			NormalUserName += FChar::ToUpper(UserName[Idx]);
+		}
+		else if(UserName[Idx] == '.')
+		{
+			NormalUserName += ' ';
+		}
+		else
+		{
+			NormalUserName += UserName[Idx];
+		}
+	}
+	return NormalUserName;
+}
+
+int FUtility::ExecuteProcess(const TCHAR* FileName, const TCHAR* CommandLine, const TCHAR* Input, FEvent* AbortEvent, FLineBasedTextWriter& Log)
+{
+	return ExecuteProcess(FileName, CommandLine, Input, AbortEvent, [&Log](const FString& Line){ Log.WriteLine(Line); });
+}
+
+int FUtility::ExecuteProcess(const TCHAR* FileName, const TCHAR* CommandLine, const TCHAR* Input, FEvent* AbortEvent, TArray<FString>& OutLines)
+{
+	uint32 ProcId;
+	void* ReadPipe = nullptr;
+    void* WritePipe = nullptr;
+    FPlatformProcess::CreatePipe(ReadPipe, WritePipe);
+
+	FProcHandle Proc = FPlatformProcess::CreateProc(FileName, CommandLine, false, false, false, &ProcId, 0, nullptr, WritePipe, ReadPipe);
+
+	FString Output;
+    FString LatestOutput = FPlatformProcess::ReadPipe(ReadPipe);
+    while (FPlatformProcess::IsProcRunning(Proc) || !LatestOutput.IsEmpty())
+    {
+        Output += LatestOutput;
+        LatestOutput = FPlatformProcess::ReadPipe(ReadPipe);
+		FPlatformProcess::Sleep(0);
+    }
+
+    FPlatformProcess::ClosePipe(ReadPipe, WritePipe);
+
+	// TODO worried this may not always work for say p4 servers sending out only \n on Windows host
+	Output.ParseIntoArray(OutLines, LINE_TERMINATOR);
+
+	int ExitCode = -1;
+	bool GotReturnCode = FPlatformProcess::GetProcReturnCode(Proc, &ExitCode);
+
+	if (GotReturnCode)
+	{
+		return ExitCode;
+	}
+
+	return -1;
+}
+
+int FUtility::ExecuteProcess(const TCHAR* FileName, const TCHAR* CommandLine, const TCHAR* Input, FEvent* AbortEvent, const TFunction<void(const FString&)>& OutputLine)
+{
+	uint32 ProcId;
+	void* ReadPipe = nullptr;
+    void* WritePipe = nullptr;
+    FPlatformProcess::CreatePipe(ReadPipe, WritePipe);
+
+	FProcHandle Proc = FPlatformProcess::CreateProc(FileName, CommandLine, false, false, false, &ProcId, 0, nullptr, WritePipe, ReadPipe);
+
+    FString LatestOutput = FPlatformProcess::ReadPipe(ReadPipe);
+    while (FPlatformProcess::IsProcRunning(Proc) || !LatestOutput.IsEmpty())
+    {
+		// Check if we end up having a perfect data blob, with a perfect amount of lines
+		bool bEndsWithLineTerminator = LatestOutput.EndsWith(LINE_TERMINATOR);
+
+		// Parsing into this array will give us full lines, minus the last entry will be the left over data
+		TArray<FString> Lines;
+		// TODO ... worried this LINE_TERMINATOR *may* not always be correct urg. For example p4 is on a Linux server, and on Windows im not sure if it will return \r\n
+		LatestOutput.ParseIntoArray(Lines, LINE_TERMINATOR);
+
+		// If there is no data, we cannot, if theres only 1 entry or we dont have any lines found in the data chunk, move all over to the LeftOverLine
+		FString LeftOverLine;
+
+		// If we have more then 2 lines, or we only have 1 line and it didnt contain *a* LINE_TERMINATOR its left over data
+		if (Lines.Num() > 1 || (Lines.Num() > 0 && !bEndsWithLineTerminator))
+		{
+			LeftOverLine = Lines.Pop();
+		}
+
+		for (const FString& Line : Lines)
+		{
+			OutputLine(Line);
+		}
+
+		LatestOutput = LeftOverLine + FPlatformProcess::ReadPipe(ReadPipe);
+		FPlatformProcess::Sleep(0);
+    }
+
+    FPlatformProcess::ClosePipe(ReadPipe, WritePipe);
+
+	int ExitCode = -1;
+	bool GotReturnCode = FPlatformProcess::GetProcReturnCode(Proc, &ExitCode);
+
+	if (GotReturnCode)
+	{
+		return ExitCode;
+	}
+
+	return -1;
+}
+
+FString FUtility::ExpandVariables(const TCHAR* InputString, const TMap<FString, FString>* AdditionalVariables)
+{
+	FString Result = InputString;
+	for(int Idx = 0;;)
+	{
+		Idx = Result.Find(TEXT("$("), ESearchCase::CaseSensitive, ESearchDir::FromStart, Idx);
+		if(Idx == INDEX_NONE)
+		{
+			break;
+		}
+	
+		// Find the end of the variable name
+		int EndIdx = Result.Find(TEXT(")"), ESearchCase::CaseSensitive, ESearchDir::FromStart, Idx + 2);
+		if (EndIdx == INDEX_NONE)
+		{
+			break;
+		}
+
+		// Extract the variable name from the string
+		FString Name = Result.Mid(Idx + 2, EndIdx - (Idx + 2));
+
+		// Find the value for it, either from the dictionary or the environment block
+		const FString* VariableValue = nullptr;
+		if(AdditionalVariables != nullptr)
+		{
+			VariableValue = AdditionalVariables->Find(Name);
+		}
+
+		FString Value;
+		if(VariableValue != nullptr)
+		{
+			Value = *VariableValue;
+		}
+		else
+		{
+			Value = FPlatformMisc::GetEnvironmentVariable(*Name);
+			if(Value.IsEmpty())
+			{
+				Idx = EndIdx + 1;
+				continue;
+			}
+		}
+
+		// Replace the variable, or skip past it
+		Result = Result.Mid(0, Idx) + Value + Result.Mid(EndIdx + 1);
+	}
+	return Result;
+}
diff --git a/Engine/Source/Programs/SlateUGS/Private/UGSCore/Utility.h b/Engine/Source/Programs/SlateUGS/Private/UGSCore/Utility.h
new file mode 100644
--- /dev/null
+++ b/Engine/Source/Programs/SlateUGS/Private/UGSCore/Utility.h
@@ -0,0 +1,35 @@
+// Copyright 1998-2017 Epic Games, Inc. All Rights Reserved.
+
+#pragma once
+
+#include "CoreMinimal.h"
+#include "OutputAdapters.h"
+
+struct FAbortException
+{
+};
+
+struct FUtility
+{
+	static bool TryParse(const TCHAR* Text, int32& OutValue);
+	static bool TryParse(const TCHAR* Text, int64& OutValue);
+
+	static bool IsFileUnderDirectory(const TCHAR* FileName, const TCHAR* DirectoryName);
+	static FString GetPathWithCorrectCase(const FString& Path);
+
+	static FString FormatUserName(const TCHAR* UserName);
+
+	static int ExecuteProcess(const TCHAR* FileName, const TCHAR* CommandLine, const TCHAR* Input, FEvent* AbortEvent, FLineBasedTextWriter& Log);
+	static int ExecuteProcess(const TCHAR* FileName, const TCHAR* CommandLine, const TCHAR* Input, FEvent* AbortEvent, TArray<FString>& OutLines);
+	static int ExecuteProcess(const TCHAR* FileName, const TCHAR* CommandLine, const TCHAR* Input, FEvent* AbortEvent, const TFunction<void(const FString&)>& OutputLine);
+
+	/**
+	 * Expands variables in $(VarName) format in the given string. Variables are retrieved from the given dictionary, or through the environment of the current process.
+	 * Any unknown variables are ignored.
+	 *
+	 * @param InputString String to search for variable names
+	 * @param Variables Lookup of variable names to values
+	 * @return String with all variables replaced
+	 */
+	static FString ExpandVariables(const TCHAR* InputString, const TMap<FString, FString>* AdditionalVariables = nullptr);
+};
diff --git a/Engine/Source/Programs/SlateUGS/Private/UGSCore/Workspace.cpp b/Engine/Source/Programs/SlateUGS/Private/UGSCore/Workspace.cpp
new file mode 100644
--- /dev/null
+++ b/Engine/Source/Programs/SlateUGS/Private/UGSCore/Workspace.cpp
@@ -0,0 +1,1234 @@
+// Copyright 1998-2017 Epic Games, Inc. All Rights Reserved.
+
+#include "Workspace.h"
+#include "Misc/Paths.h"
+#include "Misc/ScopeLock.h"
+#include "HAL/FileManager.h"
+#include "HAL/RunnableThread.h"
+#include "HAL/PlatformProcess.h"
+#include "HAL/Event.h"
+#include "Algo/AnyOf.h"
+#include "Misc/FileHelper.h"
+#include "BuildStep.h"
+#include "Telemetry.h"
+#include "Utility.h"
+
+const TCHAR* FWorkspace::DefaultBuildTargets[] =
+{
+	TEXT("UnrealHeaderTool Win64 Development, 0.1"),
+	TEXT("$(EditorTarget) Win64 $(EditorConfiguration), 0.7"),
+	TEXT("ShaderCompileWorker Win64 Development, 0.8"),
+	TEXT("UnrealLightmass Win64 Development, 0.9"),
+	TEXT("CrashReportClient Win64 Shipping, 1.0"),
+};
+
+const FWorkspaceSyncCategory FWorkspace::DefaultSyncCategories[] =
+{
+	FWorkspaceSyncCategory(FGuid(0x6703E989, 0xD912451D, 0x93ADB48D, 0xE748D282), TEXT("Content"), TEXT("*.uasset")),
+	FWorkspaceSyncCategory(FGuid(0x6507C2FB, 0x19DD403A, 0xAFA3BBF8, 0x98248D5A), TEXT("Documentation"), TEXT("/Engine/Documentation/...")),
+	FWorkspaceSyncCategory(FGuid(0xFD7C716E, 0x4BAD43AE, 0x8FAE8748, 0xEF9EE44D), TEXT("Platform Support: Android"), TEXT("/Engine/Source/ThirdParty/.../Android/...")),
+	FWorkspaceSyncCategory(FGuid(0x3299A73D, 0x21764C0F, 0xBC99C1C6, 0x631AF6C4), TEXT("Platform Support: HTML5"), TEXT("/Engine/Source/ThirdParty/.../HTML5/...;/Engine/Extras/ThirdPartyNotUE/emsdk/...")),
+	FWorkspaceSyncCategory(FGuid(0x176B2EB2, 0x35F74E8E, 0xB1315F1C, 0x5F0959AF), TEXT("Platform Support: iOS"), TEXT("/Engine/Source/ThirdParty/.../IOS/...")),
+	FWorkspaceSyncCategory(FGuid(0xF44B2D25, 0xCBC04A8F, 0xB6B3E4A8, 0x125533DD), TEXT("Platform Support: Linux"), TEXT("/Engine/Source/ThirdParty/.../Linux/...")),
+	FWorkspaceSyncCategory(FGuid(0x2AF45231, 0x0D75463B, 0xBF9FABB3, 0x231091BB), TEXT("Platform Support: Mac"), TEXT("/Engine/Source/ThirdParty/.../Mac/...")),
+	FWorkspaceSyncCategory(FGuid(0xC8CB4934, 0xADE946C9, 0xB6E361A6, 0x59E1FAF5), TEXT("Platform Support: PS4"), TEXT(".../PS4/...")),
+	FWorkspaceSyncCategory(FGuid(0xF8AE5AC3, 0xDA2D4719, 0xBABF8A90, 0xD878379E), TEXT("Platform Support: Switch"), TEXT(".../Switch/...")),
+	FWorkspaceSyncCategory(FGuid(0x3788A0BC, 0x188C4A0D, 0x950AD681, 0x75F0D110), TEXT("Platform Support: tvOS"), TEXT("/Engine/Source/ThirdParty/.../TVOS/...")),
+	FWorkspaceSyncCategory(FGuid(0x1144E719, 0xFCD7491B, 0xB0FC8B4C, 0x3565BF79), TEXT("Platform Support: Win32"), TEXT("/Engine/Source/ThirdParty/.../Win32/...")),
+	FWorkspaceSyncCategory(FGuid(0x5206CCEE, 0x90244E36, 0x8B89F5F5, 0xA7D288D2), TEXT("Platform Support: Win64"), TEXT("/Engine/Source/ThirdParty/.../Win64/...")),
+	FWorkspaceSyncCategory(FGuid(0x06887423, 0xB0944718, 0x9B55C7A2, 0x1EE67EE4), TEXT("Platform Support: XboxOne"), TEXT(".../XboxOne/...")),
+	FWorkspaceSyncCategory(FGuid(0xCFEC942A, 0xBB904F0C, 0xACCF238E, 0xCAAD9430), TEXT("Source Code"), TEXT("/Engine/Source/...")),
+};
+
+const TCHAR* FWorkspace::BuildVersionFileName = TEXT("/Engine/Build/Build.version");
+const TCHAR* FWorkspace::VersionHeaderFileName = TEXT("/Engine/Source/Runtime/Launch/Resources/Version.h");
+const TCHAR* FWorkspace::ObjectVersionFileName = TEXT("/Engine/Source/Runtime/Core/Private/UObject/ObjectVersion.cpp");
+
+FWorkspace* FWorkspace::ActiveWorkspace = nullptr;
+
+//// EWorkspaceUpdateResult ////
+
+FString ToString(EWorkspaceUpdateResult WorkspaceUpdateResult)
+{
+	switch(WorkspaceUpdateResult)
+	{
+	case EWorkspaceUpdateResult::Canceled:
+		return TEXT("Canceled");
+	case EWorkspaceUpdateResult::FailedToSync:
+		return TEXT("FailedToSync");
+	case EWorkspaceUpdateResult::FilesToResolve:
+		return TEXT("FilesToResolve");
+	case EWorkspaceUpdateResult::FilesToClobber:
+		return TEXT("FilesToClobber");
+	case EWorkspaceUpdateResult::FailedToCompile:
+		return TEXT("FailedToCompile");
+	case EWorkspaceUpdateResult::FailedToCompileWithCleanWorkspace:
+		return TEXT("FailedToCompileWithCleanWorkspace");
+	case EWorkspaceUpdateResult::Success:
+		return TEXT("Success");
+	default:
+		check(false);
+		return FString();
+	}
+}
+
+bool TryParse(const TCHAR* Text, EWorkspaceUpdateResult& OutWorkspaceUpdateResult)
+{
+	if(FCString::Stricmp(Text, TEXT("Canceled")) == 0)
+	{
+		OutWorkspaceUpdateResult = EWorkspaceUpdateResult::Canceled;
+		return true;
+	}
+	if(FCString::Stricmp(Text, TEXT("FailedToSync")) == 0)
+	{
+		OutWorkspaceUpdateResult = EWorkspaceUpdateResult::FailedToSync;
+		return true;
+	}
+	if(FCString::Stricmp(Text, TEXT("FilesToResolve")) == 0)
+	{
+		OutWorkspaceUpdateResult = EWorkspaceUpdateResult::FilesToResolve;
+		return true;
+	}
+	if(FCString::Stricmp(Text, TEXT("FilesToClobber")) == 0)
+	{
+		OutWorkspaceUpdateResult = EWorkspaceUpdateResult::FilesToClobber;
+		return true;
+	}
+	if(FCString::Stricmp(Text, TEXT("FailedToCompile")) == 0)
+	{
+		OutWorkspaceUpdateResult = EWorkspaceUpdateResult::FailedToCompile;
+		return true;
+	}
+	if(FCString::Stricmp(Text, TEXT("FailedToCompileWithCleanWorkspace")) == 0)
+	{
+		OutWorkspaceUpdateResult = EWorkspaceUpdateResult::FailedToCompileWithCleanWorkspace;
+		return true;
+	}
+	if(FCString::Stricmp(Text, TEXT("Success")) == 0)
+	{
+		OutWorkspaceUpdateResult = EWorkspaceUpdateResult::Success;
+		return true;
+	}
+	return false;
+}
+
+//// FWorkspaceUpdateContext ////
+
+FWorkspaceUpdateContext::FWorkspaceUpdateContext(int InChangeNumber, EWorkspaceUpdateOptions InOptions, const TArray<FString>& InSyncFilter, const TMap<FGuid, FCustomConfigObject>& InDefaultBuildSteps, const TArray<FCustomConfigObject>& InUserBuildSteps, const TSet<FGuid>& InCustomBuildSteps, const TMap<FString, FString>& InVariables)
+	: StartTime(FDateTime::UtcNow())
+	, ChangeNumber(InChangeNumber)
+	, Options(InOptions)
+	, SyncFilter(InSyncFilter)
+	, DefaultBuildSteps(InDefaultBuildSteps)
+	, UserBuildStepObjects(InUserBuildSteps)
+	, CustomBuildSteps(InCustomBuildSteps)
+	, Variables(InVariables)
+{
+}
+
+//// FWorkspaceSyncCategory ////
+
+FWorkspaceSyncCategory::FWorkspaceSyncCategory(const FGuid& InUniqueId)
+	: UniqueId(InUniqueId)
+	, bEnable(true)
+	, Name(TEXT("Unnamed"))
+{
+}
+
+FWorkspaceSyncCategory::FWorkspaceSyncCategory(const FGuid& InUniqueId, const TCHAR* InName, const TCHAR* InPaths)
+	: UniqueId(InUniqueId)
+	, bEnable(true)
+	, Name(InName)
+{
+	FString(InPaths).ParseIntoArray(Paths, TEXT(";"));
+}
+
+//// FWorkspace ////
+
+FWorkspace::FWorkspace(TSharedRef<FPerforceConnection> InPerforce, const FString& InLocalRootPath, const FString& InSelectedLocalFileName, const FString& InClientRootPath, const FString& InSelectedClientFileName, int InInitialChangeNumber, int InLastBuiltChangeNumber, const FString& InTelemetryProjectPath, TSharedRef<FLineBasedTextWriter> InLog)
+	: Perforce(InPerforce)
+	, SyncPaths(GetSyncPaths(InClientRootPath, InSelectedClientFileName))
+	, LocalRootPath(InLocalRootPath)
+	, SelectedLocalFileName(InSelectedLocalFileName)
+	, ClientRootPath(InClientRootPath)
+	, SelectedClientFileName(InSelectedClientFileName)
+	, TelemetryProjectPath(InTelemetryProjectPath)
+	, CurrentChangeNumber(InInitialChangeNumber)
+	, PendingChangeNumber(InInitialChangeNumber)
+	, LastBuiltChangeNumber(InLastBuiltChangeNumber)
+	, Log(InLog)
+	, bSyncing(false)
+	, ProjectConfigFile(ReadProjectConfigFile(InLocalRootPath, InSelectedLocalFileName, InLog.Get()))
+	, AbortEvent(FPlatformProcess::GetSynchEventFromPool(true))
+	, WorkerThread(nullptr)
+{
+	ProjectStreamFilter = ReadProjectStreamFilter(InPerforce.Get(), ProjectConfigFile.Get(), AbortEvent, InLog.Get());
+}
+
+FWorkspace::~FWorkspace()
+{
+	CancelUpdate();
+}
+
+TMap<FGuid, FWorkspaceSyncCategory> FWorkspace::GetSyncCategories() const
+{
+	TMap<FGuid, FWorkspaceSyncCategory> UniqueIdToCategory;
+
+	// Add the default filters
+	for(const FWorkspaceSyncCategory& DefaultSyncCategory : DefaultSyncCategories)
+	{
+		UniqueIdToCategory.Add(DefaultSyncCategory.UniqueId, FWorkspaceSyncCategory(DefaultSyncCategory));
+	}
+
+	// Add the custom filters
+	TArray<FString> CategoryLines;
+	if(ProjectConfigFile->TryGetValues(TEXT("Options.SyncCategory"), CategoryLines))
+	{
+		for(const FString& CategoryLine : CategoryLines)
+		{
+			const FCustomConfigObject Object(*CategoryLine);
+
+			FGuid UniqueId;
+			if(Object.TryGetValue(TEXT("UniqueId"), UniqueId))
+			{
+				FWorkspaceSyncCategory* Category = UniqueIdToCategory.Find(UniqueId);
+				if(Category == nullptr)
+				{
+					Category = &(UniqueIdToCategory.Add(UniqueId, FWorkspaceSyncCategory(UniqueId)));
+				}
+
+				if(Object.GetValueOrDefault(TEXT("Clear"), false))
+				{
+					Category->Paths.Empty();
+				}
+
+				Category->Name = Object.GetValueOrDefault(TEXT("Name"), *Category->Name);
+				Category->bEnable = Object.GetValueOrDefault(TEXT("Enable"), Category->bEnable);
+
+				FString ObjectPaths;
+				if(Object.TryGetValue(TEXT("Paths"), ObjectPaths))
+				{
+					TArray<FString> NewPaths;
+					ObjectPaths.ParseIntoArray(NewPaths, TEXT(";"));
+
+					Category->Paths.Append(NewPaths);
+					Category->Paths.Sort();
+
+					for(int Idx = Category->Paths.Num() - 1; Idx > 0; Idx--)
+					{
+						if(Category->Paths[Idx - 1] == Category->Paths[Idx])
+						{
+							Category->Paths.RemoveAt(Idx);
+						}
+					}
+				}
+			}
+		}
+	}
+
+	return UniqueIdToCategory;
+}
+
+TSharedPtr<FCustomConfigFile, ESPMode::ThreadSafe> FWorkspace::GetProjectConfigFile() const
+{
+	FScopeLock Lock(&CriticalSection);
+	return ProjectConfigFile;
+}
+
+void FWorkspace::GetProjectStreamFilter(TArray<FString>& Filter)
+{
+	FScopeLock Lock(&CriticalSection);
+	Filter = ProjectStreamFilter;
+}
+
+bool FWorkspace::IsBusy() const
+{
+	return bSyncing;
+}
+
+TTuple<FString, float> FWorkspace::GetCurrentProgress() const
+{
+	FScopeLock Lock(&CriticalSection);
+	return Progress.GetCurrent();
+}
+
+int FWorkspace::GetCurrentChangeNumber() const
+{
+	return CurrentChangeNumber;
+}
+
+int FWorkspace::GetPendingChangeNumber() const
+{
+	return PendingChangeNumber;
+}
+
+int FWorkspace::GetLastBuiltChangeNumber() const
+{
+	return LastBuiltChangeNumber;
+}
+
+FString FWorkspace::GetClientName() const
+{
+	return Perforce->ClientName;
+}
+
+void FWorkspace::Update(const TSharedRef<FWorkspaceUpdateContext, ESPMode::ThreadSafe>& Context)
+{
+	// Kill any existing sync
+	CancelUpdate();
+
+	// Set the initial progress message
+	if(CurrentChangeNumber != Context->ChangeNumber)
+	{
+		PendingChangeNumber = Context->ChangeNumber;
+		if(!EnumHasAnyFlags(Context->Options, EWorkspaceUpdateOptions::SyncSingleChange))
+		{
+			CurrentChangeNumber = -1;
+		}
+	}
+	Progress.Clear();
+	bSyncing = true;
+
+	// Spawn the new thread
+	AbortEvent->Reset();
+	WorkerThreadContext = Context;
+	WorkerThread = FRunnableThread::Create(this, *FString::Printf(TEXT("Update Workspace: %s"), *SelectedClientFileName));
+}
+
+void FWorkspace::CancelUpdate()
+{
+	if(bSyncing)
+	{
+		Log->WriteLine(TEXT("OPERATION ABORTED"));
+		if(WorkerThread != nullptr)
+		{
+			AbortEvent->Trigger();
+			WorkerThread->WaitForCompletion();
+			WorkerThread = nullptr;
+		}
+		PendingChangeNumber = CurrentChangeNumber;
+		WorkerThreadContext.Reset();
+		bSyncing = false;
+		FPlatformAtomics::InterlockedCompareExchangePointer((void**)&ActiveWorkspace, nullptr, this);
+	}
+}
+
+uint32 FWorkspace::Run()
+{
+	FString StatusMessage;
+
+	EWorkspaceUpdateResult Result = EWorkspaceUpdateResult::FailedToSync;
+	// TODO enable exceptions or dont do this
+	//try
+	{
+		Result = UpdateWorkspaceInternal(*WorkerThreadContext.Get(), StatusMessage);
+		if(Result != EWorkspaceUpdateResult::Success)
+		{
+			Log->Logf(TEXT("%s"), *StatusMessage);
+		}
+	}
+	//catch(FAbortException)
+	{
+		//StatusMessage = "Canceled.";
+		//Log->Logf(TEXT("Canceled."));
+	}
+//	catch(Exception Ex)
+//	{
+//		StatusMessage = "Failed with exception - " + Ex.ToString();
+//		Log.WriteException(Ex, "Failed with exception");
+//	}
+
+	bSyncing = false;
+	PendingChangeNumber = CurrentChangeNumber;
+	FPlatformAtomics::InterlockedCompareExchangePointer((void**)&ActiveWorkspace, nullptr, this);
+
+//	if(OnUpdateComplete != null)
+//	{
+//		OnUpdateComplete(Context, Result, StatusMessage);
+//	}
+
+	WorkerThreadContext.Reset();
+	return 0;
+}
+
+EWorkspaceUpdateResult FWorkspace::UpdateWorkspaceInternal(FWorkspaceUpdateContext& Context, FString& OutStatusMessage)
+{
+	IFileManager& FileManager = IFileManager::Get();
+
+//	string CmdExe = Path.Combine(Environment.GetFolderPath(Environment.SpecialFolder.System), "cmd.exe");
+#if PLATFORM_WINDOWS
+	FString CmdExe = TEXT("C:\\Windows\\System32\\cmd.exe");
+#else
+	FString CmdExe = TEXT("/usr/bin/env");
+#endif
+	//if(!FileManager.FileExists(*CmdExe))
+	//{
+		//OutStatusMessage = FString::Printf(TEXT("Missing %s."), *CmdExe);
+		//return EWorkspaceUpdateResult::FailedToSync;
+	//}
+
+	TArray<TTuple<FString, FTimespan>> Times;
+
+	int NumFilesSynced = 0;
+	if(EnumHasAnyFlags(Context.Options, EWorkspaceUpdateOptions::Sync) || EnumHasAnyFlags(Context.Options, EWorkspaceUpdateOptions::SyncSingleChange))
+	{
+//		using(TelemetryStopwatch Stopwatch = new TelemetryStopwatch("Sync", TelemetryProjectPath))
+		{
+			Log->Logf(TEXT("Syncing to %d..."), PendingChangeNumber);
+
+			// Find all the files that are out of date
+			Progress.Set(TEXT("Finding files to sync..."));
+
+			// Get the user's sync filter
+			TUniquePtr<FFileFilter> UserFilter;
+			if(Context.SyncFilter.Num() > 0)
+			{
+				UserFilter = TUniquePtr<FFileFilter>(new FFileFilter(EFileFilterType::Include));
+				for(const FString& UserFilterLine : Context.SyncFilter)
+				{
+					FString Line = UserFilterLine.TrimStartAndEnd();
+					if(Line.Len() > 0 && Line[0] != ';' && Line[0] != '#')
+					{
+						UserFilter->AddRule(Line);
+					}
+				}
+			}
+
+			// Find all the server changes, and anything that's opened for edit locally. We need to sync files we have open to schedule a resolve.
+			TArray<FString> SyncFiles;
+			for(const FString& SyncPath : SyncPaths)
+			{
+				TArray<FPerforceFileRecord> SyncRecords;
+				if(!Perforce->SyncPreview(SyncPath, PendingChangeNumber, !EnumHasAnyFlags(Context.Options, EWorkspaceUpdateOptions::Sync), SyncRecords, AbortEvent, Log.Get()))
+				{
+					OutStatusMessage = FString::Printf(TEXT("Couldn't enumerate changes matching %s."), *SyncPath);
+					return EWorkspaceUpdateResult::FailedToSync;
+				}
+
+				if(UserFilter.IsValid())
+				{
+//					SyncRecords.RemoveAll([&UserFilter](const FPerforceFileRecord& SyncRecord){ return SyncRecord.ClientPath.Len() > 0 && !MatchFilter(FPaths::Con.GetFullPath(x.ClientPath), UserFilter));
+				}
+
+				for(const FPerforceFileRecord& SyncRecord : SyncRecords)
+				{
+					SyncFiles.Add(SyncRecord.DepotPath);
+				}
+
+				TArray<FPerforceFileRecord> OpenRecords;
+				if(!Perforce->GetOpenFiles(SyncPath, OpenRecords, AbortEvent, Log.Get()))
+				{
+					OutStatusMessage = FString::Printf(TEXT("Couldn't find open files matching %s."), *SyncPath);
+					return EWorkspaceUpdateResult::FailedToSync;
+				}
+
+				// don't force a sync on added files
+				for(const FPerforceFileRecord& OpenRecord : OpenRecords)
+				{
+					if(OpenRecord.Action != TEXT("add") && OpenRecord.Action != TEXT("branch") && OpenRecord.Action != TEXT("move/add"))
+					{
+						SyncFiles.Add(OpenRecord.DepotPath);
+					}
+				}
+			}
+
+			// Filter out all the binaries that we don't want
+			FFileFilter Filter(EFileFilterType::Include);
+			Filter.AddRule(FString::Printf(TEXT("...%s"), BuildVersionFileName), EFileFilterType::Exclude);
+			Filter.AddRule(FString::Printf(TEXT("...%s"), VersionHeaderFileName), EFileFilterType::Exclude);
+			Filter.AddRule(FString::Printf(TEXT("...%s"), ObjectVersionFileName), EFileFilterType::Exclude);
+			if(EnumHasAnyFlags(Context.Options, EWorkspaceUpdateOptions::ContentOnly))
+			{
+				Filter.AddRule(TEXT("*.usf"), EFileFilterType::Exclude);
+			}
+			SyncFiles.RemoveAll([&Filter](const FString& FileName){ return !Filter.Matches(FileName); });
+
+			// Sync them all
+			TArray<FString> TamperedFiles;
+			TSet<FString> RemainingFiles(SyncFiles);
+			if(!Perforce->Sync(SyncFiles, PendingChangeNumber, [this, &RemainingFiles, &SyncFiles](const FPerforceFileRecord& Record){ UpdateSyncProgress(Record, RemainingFiles, SyncFiles.Num()); }, TamperedFiles, &Context.PerforceSyncOptions, AbortEvent, Log.Get()))
+			{
+				OutStatusMessage = TEXT("Aborted sync due to errors.");
+				return EWorkspaceUpdateResult::FailedToSync;
+			}
+
+			// If any files need to be clobbered, defer to the main thread to figure out which ones
+			if(TamperedFiles.Num() > 0)
+			{
+				int NumNewFilesToClobber = 0;
+				for(const FString& TamperedFile : TamperedFiles)
+				{
+					if(!Context.ClobberFiles.Contains(TamperedFile))
+					{
+						Context.ClobberFiles[TamperedFile] = true;
+						NumNewFilesToClobber++;
+					}
+				}
+				if(NumNewFilesToClobber > 0)
+				{
+					OutStatusMessage = FString::Printf(TEXT("Cancelled sync after checking files to clobber (%d new files)."), NumNewFilesToClobber);
+					return EWorkspaceUpdateResult::FilesToClobber;
+				}
+				for(const FString& TamperedFile : TamperedFiles)
+				{
+					if(Context.ClobberFiles[TamperedFile] && !Perforce->ForceSync(TamperedFile, PendingChangeNumber, AbortEvent, Log.Get()))
+					{
+						OutStatusMessage = FString::Printf(TEXT("Couldn't sync %s."), *TamperedFile);
+						return EWorkspaceUpdateResult::FailedToSync;
+					}
+				}
+			}
+
+			int VersionChangeNumber = -1;
+			if(EnumHasAnyFlags(Context.Options, EWorkspaceUpdateOptions::Sync))
+			{
+				// Read the new config file
+				ProjectConfigFile = ReadProjectConfigFile(LocalRootPath, SelectedLocalFileName, Log.Get());
+				ProjectStreamFilter = ReadProjectStreamFilter(Perforce.Get(), ProjectConfigFile.Get(), AbortEvent, Log.Get());
+
+				// Get the branch name
+				FString BranchOrStreamName;
+				if(!Perforce->GetActiveStream(BranchOrStreamName, AbortEvent, Log.Get()))
+				{
+					FString DepotFileName;
+					if(!Perforce->ConvertToDepotPath(ClientRootPath + TEXT("/GenerateProjectFiles.sh"), DepotFileName, AbortEvent, Log.Get()))
+					{
+						OutStatusMessage = FString::Printf(TEXT("Couldn't determine branch name for %s."), *SelectedClientFileName);
+						return EWorkspaceUpdateResult::FailedToSync;
+					}
+					BranchOrStreamName = FPerforceUtils::GetClientOrDepotDirectoryName(*DepotFileName);
+				}
+
+				// Get all the paths for code changes
+				TArray<FString> CodeFilters;
+				for(const FString& SyncPath : SyncPaths)
+				{
+					const TCHAR* CodeExtensions[] = { TEXT(".cs"), TEXT(".h"), TEXT(".cpp"), TEXT(".usf") };
+					for(const TCHAR* CodeExtension : CodeExtensions)
+					{
+						CodeFilters.Add(FString::Printf(TEXT("%s%s@<=%d"), *SyncPath, CodeExtension, PendingChangeNumber));
+					}
+				}
+
+				// Find the last code change before this changelist. For consistency in versioning between local builds and precompiled binaries, we need to use the last submitted code changelist as our version number.
+				TArray<FPerforceChangeSummary> CodeChanges;
+				if(!Perforce->FindChanges(CodeFilters, 1, CodeChanges, AbortEvent, Log.Get()))
+				{
+					OutStatusMessage = FString::Printf(TEXT("Couldn't determine last code changelist before CL %d."), PendingChangeNumber);
+					return EWorkspaceUpdateResult::FailedToSync;
+				}
+				if(CodeChanges.Num() == 0)
+				{
+					OutStatusMessage = FString::Printf(TEXT("Could not find any code changes before CL %d."), PendingChangeNumber);
+					return EWorkspaceUpdateResult::FailedToSync;
+				}
+
+				// Get the last code change
+				if(ProjectConfigFile->GetValue(TEXT("Options.VersionToLastCodeChange"), true))
+				{
+					for(const FPerforceChangeSummary& CodeChange : CodeChanges)
+					{
+						VersionChangeNumber = FMath::Max(VersionChangeNumber, CodeChange.Number);
+					}
+				}
+				else
+				{
+					VersionChangeNumber = PendingChangeNumber;
+				}
+
+				// Update the version files
+				if(ProjectConfigFile->GetValue(TEXT("Options.UseFastModularVersioning"), false))
+				{
+					TMap<FString, FString> BuildVersionStrings;
+					BuildVersionStrings.Add(TEXT("\"Changelist\":"), FString::Printf(TEXT(" %d,"), PendingChangeNumber));
+					BuildVersionStrings.Add(TEXT("\"CompatibleChangelist\":"), FString::Printf(TEXT(" %d,"), VersionChangeNumber));
+					BuildVersionStrings.Add(TEXT("\"BranchName\":"), FString::Printf(TEXT(" \"%s\""), *BranchOrStreamName.Replace(TEXT("/"), TEXT("+"))));
+					BuildVersionStrings.Add(TEXT("\"IsPromotedBuild\":"), TEXT(" 0,"));
+					if(!UpdateVersionFile(*(ClientRootPath + BuildVersionFileName), BuildVersionStrings, PendingChangeNumber))
+					{
+						OutStatusMessage = FString::Printf(TEXT("Failed to update %s."), *BuildVersionFileName);
+						return EWorkspaceUpdateResult::FailedToSync;
+					}
+
+					TMap<FString, FString> VersionHeaderStrings;
+					VersionHeaderStrings.Add(TEXT("#define ENGINE_IS_PROMOTED_BUILD"), TEXT(" (0)"));
+					VersionHeaderStrings.Add(TEXT("#define BUILT_FROM_CHANGELIST"), TEXT(" 0"));
+					VersionHeaderStrings.Add(TEXT("#define BRANCH_NAME"), TEXT(" \"") + BranchOrStreamName.Replace(TEXT("/"), TEXT("+")) + TEXT("\""));
+					if(!UpdateVersionFile(*(ClientRootPath + VersionHeaderFileName), VersionHeaderStrings, PendingChangeNumber))
+					{
+						OutStatusMessage = FString::Printf(TEXT("Failed to update %s."), *VersionHeaderFileName);
+						return EWorkspaceUpdateResult::FailedToSync;
+					}
+					if(!UpdateVersionFile(*(ClientRootPath + ObjectVersionFileName), TMap<FString,FString>(), PendingChangeNumber))
+					{
+						OutStatusMessage = FString::Printf(TEXT("Failed to update %s."), *ObjectVersionFileName);
+						return EWorkspaceUpdateResult::FailedToSync;
+					}
+				}
+				else
+				{
+					if(!UpdateVersionFile(*(ClientRootPath + BuildVersionFileName), TMap<FString, FString>(), PendingChangeNumber))
+					{
+						OutStatusMessage = FString::Printf(TEXT("Failed to update %s"), *BuildVersionFileName);
+						return EWorkspaceUpdateResult::FailedToSync;
+					}
+
+					TMap<FString, FString> VersionStrings;
+					VersionStrings.Add(TEXT("#define ENGINE_VERSION"), FString::Printf(TEXT(" %d"), VersionChangeNumber));
+					VersionStrings.Add(TEXT("#define ENGINE_IS_PROMOTED_BUILD"), TEXT(" (0)"));
+					VersionStrings.Add(TEXT("#define BUILT_FROM_CHANGELIST"), FString::Printf(TEXT(" %d"), VersionChangeNumber));
+					VersionStrings.Add(TEXT("#define BRANCH_NAME"), FString::Printf(TEXT(" \"%s\""), *BranchOrStreamName.Replace(TEXT("/"), TEXT("+"))));
+					if(!UpdateVersionFile(*(ClientRootPath + VersionHeaderFileName), VersionStrings, PendingChangeNumber))
+					{
+						OutStatusMessage = FString::Printf(TEXT("Failed to update %s"), *VersionHeaderFileName);
+						return EWorkspaceUpdateResult::FailedToSync;
+					}
+					if(!UpdateVersionFile(*(ClientRootPath + ObjectVersionFileName), VersionStrings, PendingChangeNumber))
+					{
+						OutStatusMessage = FString::Printf(TEXT("Failed to update %s"), *ObjectVersionFileName);
+						return EWorkspaceUpdateResult::FailedToSync;
+					}
+				}
+
+				// Remove all the receipts for build targets in this branch
+				if(SelectedClientFileName.EndsWith(TEXT(".uproject")))
+				{
+					Perforce->Sync(FPerforceUtils::GetClientOrDepotDirectoryName(*SelectedClientFileName) + TEXT("/Build/Receipts/...#0"), AbortEvent, Log.Get());
+				}
+			}
+
+			// Check if there are any files which need resolving
+			TArray<FPerforceFileRecord> UnresolvedFiles;
+			if(!FindUnresolvedFiles(UnresolvedFiles))
+			{
+				OutStatusMessage = TEXT("Couldn't get list of unresolved files.");
+				return EWorkspaceUpdateResult::FailedToSync;
+			}
+			if(UnresolvedFiles.Num() > 0 && EnumHasAnyFlags(Context.Options, EWorkspaceUpdateOptions::AutoResolveChanges))
+			{
+				for(const FPerforceFileRecord& UnresolvedFile : UnresolvedFiles)
+				{
+					Perforce->AutoResolveFile(UnresolvedFile.DepotPath, AbortEvent, Log.Get());
+				}
+				if(!FindUnresolvedFiles(UnresolvedFiles))
+				{
+					OutStatusMessage = "Couldn't get list of unresolved files.";
+					return EWorkspaceUpdateResult::FailedToSync;
+				}
+			}
+			if(UnresolvedFiles.Num() > 0)
+			{
+				Log->Logf(TEXT("%d files need resolving:"), UnresolvedFiles.Num());
+				for(const FPerforceFileRecord& UnresolvedFile : UnresolvedFiles)
+				{
+					Log->Logf(TEXT("  %s"), *UnresolvedFile.ClientPath);
+				}
+				OutStatusMessage = TEXT("Files need resolving.");
+				return EWorkspaceUpdateResult::FilesToResolve;
+			}
+
+			// Continue processing sync-only actions
+			if (EnumHasAnyFlags(Context.Options, EWorkspaceUpdateOptions::Sync))
+			{
+				// Execute any project specific post-sync steps
+				TArray<FString> PostSyncSteps;
+				if(ProjectConfigFile->TryGetValues(TEXT("Sync.Step"), PostSyncSteps))
+				{
+					Log->Logf(TEXT(""));
+					Log->Logf(TEXT("Executing post-sync steps..."));
+
+					TMap<FString, FString> PostSyncVariables(Context.Variables);
+					PostSyncVariables.Add("Change", FString::Printf(TEXT("%d"), PendingChangeNumber));
+					PostSyncVariables.Add("CodeChange", FString::Printf(TEXT("%d"), VersionChangeNumber));
+
+					for(const FString& PostSyncStep : PostSyncSteps)
+					{
+						FCustomConfigObject PostSyncStepObject(*PostSyncStep);
+
+						FString ToolFileName = FUtility::ExpandVariables(PostSyncStepObject.GetValueOrDefault(TEXT("FileName"), TEXT("")), &PostSyncVariables);
+						if (ToolFileName.Len() > 0)
+						{
+							FString ToolArguments = FUtility::ExpandVariables(PostSyncStepObject.GetValueOrDefault(TEXT("Arguments"), TEXT("")), &PostSyncVariables);
+
+							Log->Logf(TEXT("post-sync> Running %s %s"), *ToolFileName, *ToolArguments);
+
+							FProgressTextWriter PostSyncWriter(Progress, MakeShared<FPrefixedTextWriter>(TEXT("post-sync> "), Log));
+
+							int ResultFromTool = FUtility::ExecuteProcess(*ToolFileName, *ToolArguments, NULL, AbortEvent, PostSyncWriter);
+							if (ResultFromTool != 0)
+							{
+								OutStatusMessage = FString::Printf(TEXT("Post-sync step terminated with exit code %d."), ResultFromTool);
+								return EWorkspaceUpdateResult::FailedToSync;
+							}
+						}
+					}
+				}
+
+				// Update the current change number. Everything else happens for the new change.
+				CurrentChangeNumber = PendingChangeNumber;
+			}
+
+			// Update the timing info
+//			Times.Add(TTuple<FString,FTimespan>(TEXT("Sync"), Stopwatch.Stop("Success")));
+
+			// Save the number of files synced
+			NumFilesSynced = SyncFiles.Num();
+			Log->Logf(TEXT(""));
+		}
+	}
+
+	// Extract an archive from the depot path
+	if(EnumHasAnyFlags(Context.Options, EWorkspaceUpdateOptions::SyncArchives))
+	{
+		FTelemetryStopwatch Stopwatch(TEXT("Archives"), TelemetryProjectPath);
+
+		// Create the directory for extracted archive manifests
+		FString ManifestDirectoryName;
+		if(SelectedLocalFileName.EndsWith(TEXT(".uproject")))
+		{
+			ManifestDirectoryName = FPaths::GetPath(SelectedLocalFileName) / TEXT("Saved/UnrealGameSync");
+		}
+		else
+		{
+			ManifestDirectoryName = FPaths::GetPath(SelectedLocalFileName) / TEXT("Engine/Saved/UnrealGameSync");
+		}
+		IFileManager::Get().MakeDirectory(*ManifestDirectoryName);
+
+		// Sync and extract (or just remove) the given archives
+		for(const TTuple<FString, FString>& ArchiveTypeAndDepotPath : Context.ArchiveTypeToDepotPath)
+		{
+			// Remove any existing binaries
+			FString ManifestFileName = ManifestDirectoryName / FString::Printf(TEXT("%s.zipmanifest"), *ArchiveTypeAndDepotPath.Key);
+			if(IFileManager::Get().FileExists(*ManifestFileName))
+			{
+				Log->Logf(TEXT("Removing %s binaries..."), *ArchiveTypeAndDepotPath.Key);
+				Progress.Set(*FString::Printf(TEXT("Removing %s binaries..."), *ArchiveTypeAndDepotPath.Key), 0.0f);
+//					ArchiveUtils.RemoveExtractedFiles(LocalRootPath, ManifestFileName, Progress, Log);
+				Log->Logf(TEXT(""));
+			}
+
+			// If we have a new depot path, sync it down and extract it
+			if(ArchiveTypeAndDepotPath.Value.Len() > 0)
+			{
+//					string TempZipFileName = Path.GetTempFileName();
+//					try
+//					{
+//						Log.WriteLine("Syncing {0} binaries...", ArchiveTypeAndDepotPath.Key.ToLowerInvariant());
+//						Progress.Set(String.Format("Syncing {0} binaries...", ArchiveTypeAndDepotPath.Key.ToLowerInvariant()), 0.0f);
+//						if(!Perforce.PrintToFile(ArchiveTypeAndDepotPath.Value, TempZipFileName, Log) || new FileInfo(TempZipFileName).Length == 0)
+//						{
+//							StatusMessage = String.Format("Couldn't read {0}", ArchiveTypeAndDepotPath.Value);
+//							return WorkspaceUpdateResult.FailedToSync;
+//						}
+//						ArchiveUtils.ExtractFiles(TempZipFileName, LocalRootPath, ManifestFileName, Progress, Log);
+//						Log.WriteLine();
+//					}
+//					finally
+//					{
+//						File.SetAttributes(TempZipFileName, FileAttributes.Normal);
+//						File.Delete(TempZipFileName);
+//					}
+			}
+		}
+
+		// Add the finish time
+		Times.Add(TTuple<FString,FTimespan>(TEXT("Archive"), Stopwatch.Stop(TEXT("Success"))));
+	}
+
+	// Take the lock before doing anything else. Building and generating project files can only be done on one workspace at a time.
+	if(FPlatformAtomics::InterlockedCompareExchangePointer((void**)&ActiveWorkspace, this, nullptr) != nullptr)
+	{
+		Log->Logf(TEXT("Waiting for other workspaces to finish..."));
+		while(FPlatformAtomics::InterlockedCompareExchangePointer((void**)&ActiveWorkspace, this, nullptr) != nullptr)
+		{
+			FPlatformProcess::Sleep(0.1f);
+		}
+	}
+
+	// Generate project files in the workspace
+	if(EnumHasAnyFlags(Context.Options, EWorkspaceUpdateOptions::GenerateProjectFiles))
+	{
+		FTelemetryStopwatch Stopwatch(TEXT("Prj gen"), TelemetryProjectPath);
+		Progress.Set(TEXT("Generating project files..."), 0.0f);
+
+		FString ProjectFileArgument;
+		if(SelectedLocalFileName.EndsWith(TEXT(".uproject")))
+		{
+			ProjectFileArgument = FString::Printf(TEXT("\"%s\" "), *SelectedLocalFileName);
+		}
+
+#if PLATFORM_WINDOWS
+		FString CommandLine = FString::Printf(TEXT("/C \"\"%s\" %s-progress\""), *(LocalRootPath / TEXT("GenerateProjectFiles.sh")), *ProjectFileArgument);
+#else
+		FString CommandLine = FString::Printf(TEXT("\"%s\" %s-progress"), *(LocalRootPath / TEXT("GenerateProjectFiles.sh")), *ProjectFileArgument);
+#endif
+
+		Log->Logf(TEXT("Generating project files..."));
+		Log->Logf(TEXT("gpf> Running %s %s"), *CmdExe, *CommandLine);
+
+		FProgressTextWriter ProjectFilesWriter(Progress, MakeShared<FPrefixedTextWriter>(TEXT("gpf> "), Log));
+
+		int GenerateProjectFilesResult = FUtility::ExecuteProcess(*CmdExe, *CommandLine, nullptr, AbortEvent, ProjectFilesWriter);
+		if(GenerateProjectFilesResult != 0)
+		{
+			OutStatusMessage = FString::Printf(TEXT("Failed to generate project files (exit code %d)."), GenerateProjectFilesResult);
+			return EWorkspaceUpdateResult::FailedToCompile;
+		}
+
+		Log->Logf(TEXT(""));
+		Times.Add(TTuple<FString,FTimespan>(TEXT("Prj gen"), Stopwatch.Stop(TEXT("Success"))));
+	}
+
+	// Build everything using MegaXGE
+	if(EnumHasAnyFlags(Context.Options, EWorkspaceUpdateOptions::Build))
+	{
+		// Get all the project build steps
+		TArray<FCustomConfigObject> ProjectBuildSteps;
+		TArray<FString> ProjectBuildStepStrings;
+		if(ProjectConfigFile->TryGetValues(TEXT("Build.Step"), ProjectBuildStepStrings))
+		{
+			for(const FString& ProjectBuildStepString : ProjectBuildStepStrings)
+			{
+				ProjectBuildSteps.Add(FCustomConfigObject(*ProjectBuildStepString));
+			}
+		}
+
+		// Compile all the build steps together
+		TMap<FGuid, FCustomConfigObject> BuildStepObjects = Context.DefaultBuildSteps;
+		FBuildStep::MergeBuildStepObjects(BuildStepObjects, ProjectBuildSteps);
+		FBuildStep::MergeBuildStepObjects(BuildStepObjects, Context.UserBuildStepObjects);
+
+		// Construct build steps from them
+		TArray<FBuildStep> BuildSteps;
+		for(const TTuple<FGuid, FCustomConfigObject>& BuildStepObject : BuildStepObjects)
+		{
+			BuildSteps.Add(FBuildStep(BuildStepObject.Value));
+		}
+		BuildSteps.Sort([](const FBuildStep& A, const FBuildStep& B){ return A.OrderIndex < B.OrderIndex; });
+
+		// Remove any that we're not running
+		if(Context.CustomBuildSteps.Num() > 0)
+		{
+			BuildSteps.RemoveAll([&Context](const FBuildStep& Step){ return !Context.CustomBuildSteps.Contains(Step.UniqueId); });
+		}
+		else if(EnumHasAnyFlags(Context.Options, EWorkspaceUpdateOptions::ScheduledBuild))
+		{
+			BuildSteps.RemoveAll([](const FBuildStep& Step){ return !Step.bScheduledSync; });
+		}
+		else
+		{
+			BuildSteps.RemoveAll([](const FBuildStep& Step){ return !Step.bNormalSync; });
+		}
+
+		// Check if the last successful build was before a change that we need to force a clean for
+		bool bForceClean = false;
+		if(LastBuiltChangeNumber != 0)
+		{
+			TArray<FString> CleanBuildChanges;
+			if(ProjectConfigFile->TryGetValues(TEXT("ForceClean.Changelist"), CleanBuildChanges))
+			{
+				for(const FString& CleanBuildChange : CleanBuildChanges)
+				{
+					int ChangeNumber;
+					if(FUtility::TryParse(*CleanBuildChange, ChangeNumber))
+					{
+						if((LastBuiltChangeNumber >= ChangeNumber) != (CurrentChangeNumber >= ChangeNumber))
+						{
+							Log->Logf(TEXT("Forcing clean build due to changelist %d."), ChangeNumber);
+							Log->Logf(TEXT(""));
+							bForceClean = true;
+							break;
+						}
+					}
+				}
+			}
+		}
+
+		// Execute them all
+		const TCHAR* TelemetryEventName = (Context.UserBuildStepObjects.Num() > 0)? TEXT("CustomBuild") : EnumHasAnyFlags(Context.Options, EWorkspaceUpdateOptions::UseIncrementalBuilds) ? TEXT("Compile") : TEXT("FullCompile");
+		FTelemetryStopwatch Stopwatch(TelemetryEventName, TelemetryProjectPath);
+		Progress.Set(TEXT("Starting build..."), 0.0f);
+
+		// Check we've built UBT (it should have been compiled by generating project files)
+		//FString UnrealBuildToolPath = LocalRootPath / TEXT("Engine/Binaries/DotNET/UnrealBuildTool.exe");
+		FString UnrealBuildToolPath = LocalRootPath / TEXT("Engine/Build/BatchFiles/Linux/Build.sh");
+		if(!IFileManager::Get().FileExists(*UnrealBuildToolPath))
+		{
+			OutStatusMessage = FString::Printf(TEXT("Couldn't find %s"), *UnrealBuildToolPath);
+			return EWorkspaceUpdateResult::FailedToCompile;
+		}
+
+		// Calculate the total estimated duration
+		float TotalEstimatedDuration = 0.0f;
+		for(const FBuildStep& Step : BuildSteps)
+		{
+			TotalEstimatedDuration += Step.EstimatedDuration;
+		}
+
+		// Execute all the steps
+		float MaxProgressFraction = 0.0f;
+		for(const FBuildStep& Step : BuildSteps)
+		{
+			MaxProgressFraction += (float)Step.EstimatedDuration / (float)FMath::Max(TotalEstimatedDuration, 1.0f);
+
+			Progress.Set(*Step.StatusText);
+			Progress.Push(MaxProgressFraction);
+
+			Log->Logf(TEXT("%s"), *Step.StatusText);
+
+			switch(Step.Type)
+			{
+			case EBuildStepType::Compile:
+				{
+					FTelemetryStopwatch StepStopwatch(FString::Printf(TEXT("Compile:%s"), *Step.Target), TelemetryProjectPath);
+
+					FString CommandLine = FString::Printf(TEXT("%s %s %s %s -NoHotReloadFromIDE"), *Step.Target, *Step.Platform, *Step.Configuration, *FUtility::ExpandVariables(*Step.Arguments, &Context.Variables));
+					if(!EnumHasAnyFlags(Context.Options, EWorkspaceUpdateOptions::UseIncrementalBuilds) || bForceClean)
+					{
+						Log->Logf(TEXT("ubt> Running %s %s -clean"), *UnrealBuildToolPath, *CommandLine);
+						FProgressTextWriter CleanWriter(Progress, MakeShared<FPrefixedTextWriter>(TEXT("ubt> "), Log));
+						FUtility::ExecuteProcess(*UnrealBuildToolPath, *(CommandLine + TEXT(" -clean")), nullptr, AbortEvent, CleanWriter);
+					}
+
+					Log->Logf(TEXT("ubt> Running %s %s -progress"), *UnrealBuildToolPath, *CommandLine);
+
+					FProgressTextWriter BuildWriter(Progress, MakeShared<FPrefixedTextWriter>(TEXT("ubt> "), Log));
+					int ResultFromBuild = FUtility::ExecuteProcess(*UnrealBuildToolPath, *(CommandLine + TEXT(" -progress")), nullptr, AbortEvent, BuildWriter);
+					if(ResultFromBuild != 0)
+					{
+						StepStopwatch.Stop(TEXT("Failed"));
+						OutStatusMessage = FString::Printf(TEXT("Failed to compile %s."), *Step.Target);
+						return (HasModifiedSourceFiles() || Context.UserBuildStepObjects.Num() > 0)? EWorkspaceUpdateResult::FailedToCompile : EWorkspaceUpdateResult::FailedToCompileWithCleanWorkspace;
+					}
+
+					StepStopwatch.Stop(TEXT("Success"));
+				}
+				break;
+			case EBuildStepType::Cook:
+				{
+					FTelemetryStopwatch StepStopwatch(FString::Printf(TEXT("Cook/Launch: %s"), *FPaths::GetBaseFilename(Step.FileName)), TelemetryProjectPath);
+
+					FString LocalRunUAT = LocalRootPath / TEXT("Engine/Build/BatchFiles/RunUAT.sh");
+#if PLATFORM_WINDOWS
+					FString Arguments = FString::Printf(TEXT("/C \"\"%s\" -profile=\"%s\"\""), *LocalRunUAT, *(LocalRootPath / Step.FileName));
+#else
+					FString Arguments = FString::Printf(TEXT("\"%s\" -profile=\"%s\""), *LocalRunUAT, *(LocalRootPath / Step.FileName));
+#endif
+					Log->Logf(TEXT("uat> Running %s %s"), *LocalRunUAT, *Arguments);
+
+					FProgressTextWriter CookLogWriter(Progress, MakeShared<FPrefixedTextWriter>(TEXT("uat> "), Log));
+
+					int ResultFromUAT = FUtility::ExecuteProcess(*CmdExe, *Arguments, nullptr, AbortEvent, CookLogWriter);
+					if(ResultFromUAT != 0)
+					{
+						StepStopwatch.Stop(TEXT("Failed"));
+						OutStatusMessage = FString::Printf(TEXT("Cook failed. (%d)"), ResultFromUAT);
+						return EWorkspaceUpdateResult::FailedToCompile;
+					}
+
+					StepStopwatch.Stop(TEXT("Success"));
+				}
+				break;
+			case EBuildStepType::Other:
+				{
+					FTelemetryStopwatch StepStopwatch(FString::Printf(TEXT("Custom: %s"), *FPaths::GetBaseFilename(Step.FileName)), TelemetryProjectPath);
+
+					FString ToolFileName = LocalRootPath / FUtility::ExpandVariables(*Step.FileName, &Context.Variables);
+					FString ToolArguments = FUtility::ExpandVariables(*Step.Arguments, &Context.Variables);
+					Log->Logf(TEXT("tool> Running %s %s"), *ToolFileName, *ToolArguments);
+
+					if(Step.bUseLogWindow)
+					{
+						FProgressTextWriter ToolWriter(Progress, MakeShared<FPrefixedTextWriter>(TEXT("tool> "), Log));
+
+						int ResultFromTool = FUtility::ExecuteProcess(*ToolFileName, *ToolArguments, nullptr, AbortEvent, ToolWriter);
+						if(ResultFromTool != 0)
+						{
+							StepStopwatch.Stop(TEXT("Failed"));
+							OutStatusMessage = FString::Printf(TEXT("Tool terminated with exit code %d."), ResultFromTool);
+							return EWorkspaceUpdateResult::FailedToCompile;
+						}
+					}
+					else
+					{
+//							using(Process.Start(ToolFileName, ToolArguments))
+//							{
+//							}
+					}
+
+					StepStopwatch.Stop(TEXT("Success"));
+				}
+				break;
+			}
+
+			Log->Logf(TEXT(""));
+			Progress.Pop();
+		}
+
+		Times.Add(TTuple<FString,FTimespan>(TEXT("Build"), Stopwatch.Stop("Success")));
+
+		// Update the last successful build change number
+		if(Context.CustomBuildSteps.Num() == 0)
+		{
+			LastBuiltChangeNumber = CurrentChangeNumber;
+		}
+	}
+
+	// Calculate the total time
+	long TotalSeconds = 0;
+	for(const TTuple<FString, FTimespan>& Time : Times)
+	{
+		TotalSeconds += (long)Time.Value.GetTotalSeconds();
+	}
+
+	// Write out all the timing information
+	Log->Logf(TEXT("Total time : %s"), *FormatTime(TotalSeconds));
+	for(const TTuple<FString, FTimespan>& Time : Times)
+	{
+		Log->Logf(TEXT("   %-8s: %s"), *Time.Key, *FormatTime((long)Time.Value.GetTotalSeconds()));
+	}
+	if(NumFilesSynced > 0)
+	{
+		Log->Logf(TEXT("{0} files synced."), NumFilesSynced);
+	}
+
+	Log->Logf(TEXT(""));
+	Log->Logf(TEXT("UPDATE SUCCEEDED"));
+
+	OutStatusMessage = TEXT("Update succeeded");
+	return EWorkspaceUpdateResult::Success;
+}
+
+
+TArray<FString> FWorkspace::GetSyncPaths(const FString& ClientRootPath, const FString& SelectedClientFileName)
+{
+	TArray<FString> SyncPaths;
+	if(SelectedClientFileName.EndsWith(TEXT(".uproject")))
+	{
+		SyncPaths.Add(ClientRootPath + TEXT("/*"));
+		SyncPaths.Add(ClientRootPath + TEXT("/Engine/..."));
+		SyncPaths.Add(FPerforceUtils::GetClientOrDepotDirectoryName(*SelectedClientFileName) + "/...");
+	}
+	else
+	{
+		SyncPaths.Add(ClientRootPath + TEXT("/..."));
+	}
+	return SyncPaths;
+}
+
+TSharedRef<FCustomConfigFile, ESPMode::ThreadSafe> FWorkspace::ReadProjectConfigFile(const FString& LocalRootPath, const FString& SelectedLocalFileName, FLineBasedTextWriter& Log)
+{
+	// Find the valid config file paths
+	TArray<FString> ProjectConfigFileNames;
+	ProjectConfigFileNames.Add(LocalRootPath / TEXT("Engine/Programs/UnrealGameSync/UnrealGameSync.ini"));
+	ProjectConfigFileNames.Add(LocalRootPath / TEXT("Engine/Programs/UnrealGameSync/NotForLicensees/UnrealGameSync.ini"));
+	if(SelectedLocalFileName.EndsWith(TEXT(".uproject")))
+	{
+		ProjectConfigFileNames.Add(FPaths::GetPath(SelectedLocalFileName) / TEXT("Build/UnrealGameSync.ini"));
+		ProjectConfigFileNames.Add(FPaths::GetPath(SelectedLocalFileName) / TEXT("Build/NotForLicensees/UnrealGameSync.ini"));
+	}
+	else
+	{
+		ProjectConfigFileNames.Add(LocalRootPath / TEXT("Engine/Programs/UnrealGameSync/DefaultProject.ini"));
+		ProjectConfigFileNames.Add(LocalRootPath / TEXT("Engine/Programs/UnrealGameSync/NotForLicensees/DefaultProject.ini"));
+	}
+
+	// Read them in
+	TSharedRef<FCustomConfigFile, ESPMode::ThreadSafe> ProjectConfig(new FCustomConfigFile());
+	for(const FString& ProjectConfigFileName : ProjectConfigFileNames)
+	{
+		if(IFileManager::Get().FileExists(*ProjectConfigFileName))
+		{
+			TArray<FString> Lines;
+			FFileHelper::LoadFileToStringArray(Lines, *ProjectConfigFileName);
+			ProjectConfig->Parse(Lines);
+			Log.WriteLine(FString::Printf(TEXT("Read config file from %s"), *ProjectConfigFileName));
+		}
+	}
+	return ProjectConfig;
+}
+
+TArray<FString> FWorkspace::ReadProjectStreamFilter(FPerforceConnection& Perforce, const FCustomConfigFile& ProjectConfigFile, FEvent* AbortEvent, FLineBasedTextWriter& Log)
+{
+	const TCHAR* StreamListDepotPath = ProjectConfigFile.GetValue(TEXT("Options.QuickSelectStreamList"), nullptr);
+	if(StreamListDepotPath == nullptr)
+	{
+		return TArray<FString>();
+	}
+
+	TArray<FString> Lines;
+	if(!Perforce.Print(StreamListDepotPath, Lines, AbortEvent, Log))
+	{
+		return TArray<FString>();
+	}
+
+	TArray<FString> FilteredLines;
+	for(const FString& Line: Lines)
+	{
+		FString TrimLine = Line.TrimStartAndEnd();
+		if(TrimLine.Len() > 0)
+		{
+			FilteredLines.Add(TrimLine);
+		}
+	}
+	return FilteredLines;
+}
+
+FString FWorkspace::FormatTime(long Seconds)
+{
+	if(Seconds >= 60)
+	{
+		return FString::Printf(TEXT("%3dm %02ds"), Seconds / 60, Seconds % 60);
+	}
+	else
+	{
+		return FString::Printf(TEXT("     %02ds"), Seconds);
+	}
+}
+
+bool FWorkspace::HasModifiedSourceFiles() const
+{
+	TArray<FPerforceFileRecord> OpenFiles;
+	if(!Perforce->GetOpenFiles(ClientRootPath + TEXT("/..."), OpenFiles, AbortEvent, Log.Get()))
+	{
+		return true;
+	}
+	if(Algo::AnyOf(OpenFiles, [](const FPerforceFileRecord& Record){ return Record.DepotPath.Contains(TEXT("/Source/")); }))
+	{
+		return true;
+	}
+	return false;
+}
+
+bool FWorkspace::FindUnresolvedFiles(TArray<FPerforceFileRecord>& OutUnresolvedFiles) const
+{
+	for(const FString& SyncPath : SyncPaths)
+	{
+		TArray<FPerforceFileRecord> Records;
+		if(!Perforce->GetUnresolvedFiles(SyncPath, Records, AbortEvent, Log.Get()))
+		{
+			Log->Logf(TEXT("Couldn't find open files matching %s"), *SyncPath);
+			return false;
+		}
+		OutUnresolvedFiles.Append(Records);
+	}
+	return true;
+}
+
+void FWorkspace::UpdateSyncProgress(const FPerforceFileRecord& Record, TSet<FString>& RemainingFiles, int NumFiles)
+{
+	RemainingFiles.Remove(Record.DepotPath);
+
+	FString Message = FString::Printf(TEXT("Syncing files... (%d/%d)"), NumFiles - RemainingFiles.Num(), NumFiles);
+	float Fraction = FMath::Min((float)(NumFiles - RemainingFiles.Num()) / (float)NumFiles, 1.0f);
+	Progress.Set(*Message, Fraction);
+
+	Log->Logf(TEXT("p4>   %s %s"), *Record.Action, *Record.ClientPath);
+}
+
+bool FWorkspace::UpdateVersionFile(const TCHAR* LocalPath, const TMap<FString, FString>& VersionStrings, int ChangeNumber) const
+{
+	TSharedPtr<FPerforceWhereRecord> WhereRecord;
+	if(!Perforce->Where(LocalPath, WhereRecord, AbortEvent, Log.Get()))
+	{
+		Log->Logf(TEXT("P4 where failed for %s"), *LocalPath);
+		return false;
+	}
+
+	TArray<FString> Lines;
+	if(!Perforce->Print(FString::Printf(TEXT("%s@%d"), *WhereRecord->DepotPath, ChangeNumber), Lines, AbortEvent, Log.Get()))
+	{
+		Log->Logf(TEXT("Couldn't get default contents of %s"), *WhereRecord->DepotPath);
+		return false;
+	}
+
+	FString NewFileContents;
+	for(const FString& Line : Lines)
+	{
+		FString NewLine = Line;
+		for(const TTuple<FString, FString>& VersionString : VersionStrings)
+		{
+			if(UpdateVersionLine(NewLine, VersionString.Key, VersionString.Value))
+			{
+				break;
+			}
+		}
+		NewFileContents += NewLine + LINE_TERMINATOR;
+	}
+
+	return WriteVersionFile(*WhereRecord, NewFileContents);
+}
+
+bool FWorkspace::WriteVersionFile(const FPerforceWhereRecord& WhereRecord, const FString& NewText) const
+{
+	IFileManager& FileManager = IFileManager::Get();
+	if(FileManager.FileExists(*WhereRecord.LocalPath))
+	{
+		FString OldText;
+		if(FFileHelper::LoadFileToString(OldText, *WhereRecord.LocalPath) && OldText == NewText)
+		{
+			Log->Logf(TEXT("Ignored %s; contents haven't changed"), *WhereRecord.LocalPath);
+			return true;
+		}
+	}
+
+	FileManager.Delete(*WhereRecord.LocalPath, false, true);
+
+	if(WhereRecord.DepotPath.Len() > 0)
+	{
+		Perforce->Sync(FString::Printf(TEXT("%s#0"), *WhereRecord.DepotPath), AbortEvent, Log.Get());
+	}
+	FFileHelper::SaveStringToFile(NewText, *WhereRecord.LocalPath);
+	Log->Logf(TEXT("Written %s"), *WhereRecord.LocalPath);
+	return true;
+}
+
+bool FWorkspace::UpdateVersionLine(FString& Line, const FString& Prefix, const FString& Suffix)
+{
+	int LineIdx = 0;
+	int PrefixIdx = 0;
+	for(;;)
+	{
+		FString PrefixToken;
+		if(!ReadToken(Prefix, PrefixIdx, PrefixToken))
+		{
+			break;
+		}
+
+		FString LineToken;
+		if(!ReadToken(Line, LineIdx, LineToken) || LineToken != PrefixToken)
+		{
+			return false;
+		}
+	}
+	Line = Line.Mid(0, LineIdx) + Suffix;
+	return true;
+}
+
+bool FWorkspace::ReadToken(const FString& Line, int32& LineIdx, FString &OutToken)
+{
+	for(;; LineIdx++)
+	{
+		if(LineIdx == Line.Len())
+		{
+			return false;
+		}
+		else if(!FChar::IsWhitespace(Line[LineIdx]))
+		{
+			break;
+		}
+	}
+
+	int StartIdx = LineIdx++;
+	if(FChar::IsAlnum(Line[StartIdx]) || Line[StartIdx] == '_')
+	{
+		while(LineIdx < Line.Len() && (FChar::IsAlnum(Line[LineIdx]) || Line[LineIdx] == '_'))
+		{
+			LineIdx++;
+		}
+	}
+	OutToken = Line.Mid(StartIdx, LineIdx - StartIdx);
+
+	return true;
+}
+
diff --git a/Engine/Source/Programs/SlateUGS/Private/UGSCore/Workspace.h b/Engine/Source/Programs/SlateUGS/Private/UGSCore/Workspace.h
new file mode 100644
--- /dev/null
+++ b/Engine/Source/Programs/SlateUGS/Private/UGSCore/Workspace.h
@@ -0,0 +1,140 @@
+// Copyright 1998-2017 Epic Games, Inc. All Rights Reserved.
+
+#pragma once
+
+#include "CoreMinimal.h"
+#include "CustomConfigFile.h"
+#include "Perforce.h"
+#include "OutputAdapters.h"
+#include "FileFilter.h"
+#include "HAL/Runnable.h"
+#include "Misc/EnumClassFlags.h"
+
+enum class EWorkspaceUpdateOptions
+{
+	Sync = 0x01,
+	SyncSingleChange = 0x02,
+	AutoResolveChanges = 0x04,
+	GenerateProjectFiles = 0x08,
+	SyncArchives = 0x10,
+	Build = 0x20,
+	UseIncrementalBuilds = 0x40,
+	ScheduledBuild = 0x80,
+	RunAfterSync = 0x100,
+	OpenSolutionAfterSync = 0x200,
+	ContentOnly = 0x400
+};
+
+ENUM_CLASS_FLAGS(EWorkspaceUpdateOptions)
+
+enum class EWorkspaceUpdateResult
+{
+	Canceled,
+	FailedToSync,
+	FilesToResolve,
+	FilesToClobber,
+	FailedToCompile,
+	FailedToCompileWithCleanWorkspace,
+	Success,
+};
+
+FString ToString(EWorkspaceUpdateResult WorkspaceUpdateResult);
+bool TryParse(const TCHAR* Text, EWorkspaceUpdateResult& OutWorkspaceUpdateResult);
+
+struct FWorkspaceUpdateContext
+{
+	FDateTime StartTime;
+	int ChangeNumber;
+	EWorkspaceUpdateOptions Options;
+	TArray<FString> SyncFilter;
+	TMap<FString, FString> ArchiveTypeToDepotPath;
+	TMap<FString, bool> ClobberFiles;
+	TMap<FGuid,FCustomConfigObject> DefaultBuildSteps;
+	TArray<FCustomConfigObject> UserBuildStepObjects;
+	TSet<FGuid> CustomBuildSteps;
+	TMap<FString, FString> Variables;
+	FPerforceSyncOptions PerforceSyncOptions;
+
+	FWorkspaceUpdateContext(int InChangeNumber, EWorkspaceUpdateOptions InOptions, const TArray<FString>& InSyncFilter, const TMap<FGuid, FCustomConfigObject>& InDefaultBuildSteps, const TArray<FCustomConfigObject>& InUserBuildSteps, const TSet<FGuid>& InCustomBuildSteps, const TMap<FString, FString>& InVariables);
+};
+
+struct FWorkspaceSyncCategory
+{
+	FGuid UniqueId;
+	bool bEnable;
+	FString Name;
+	TArray<FString> Paths;
+
+	FWorkspaceSyncCategory(const FGuid& InUniqueId);
+	FWorkspaceSyncCategory(const FGuid& InUniqueId, const TCHAR* InName, const TCHAR* InPaths);
+};
+
+class FWorkspace : FRunnable
+{
+public:
+	const TSharedRef<FPerforceConnection> Perforce;
+	const TArray<FString> SyncPaths;
+	const FString LocalRootPath;
+	const FString SelectedLocalFileName;
+	const FString ClientRootPath;
+	const FString SelectedClientFileName;
+	const FString TelemetryProjectPath;
+
+	TFunction<void(TSharedRef<FWorkspaceUpdateContext, ESPMode::ThreadSafe>, EWorkspaceUpdateResult, FString)> OnUpdateComplete;
+
+	FWorkspace(TSharedRef<FPerforceConnection> InPerforce, const FString& InLocalRootPath, const FString& InSelectedLocalFileName, const FString& InClientRootPath, const FString& InSelectedClientFileName, int InInitialChangeNumber, int InLastBuiltChangeNumber, const FString& InTelemetryProjectPath, TSharedRef<FLineBasedTextWriter> InLog);
+	~FWorkspace();
+
+	TMap<FGuid, FWorkspaceSyncCategory> GetSyncCategories() const;
+	TSharedPtr<FCustomConfigFile, ESPMode::ThreadSafe> GetProjectConfigFile() const;
+	void GetProjectStreamFilter(TArray<FString>& Filter);
+	bool IsBusy() const;
+	TTuple<FString, float> GetCurrentProgress() const;
+	int GetCurrentChangeNumber() const;
+	int GetPendingChangeNumber() const;
+	int GetLastBuiltChangeNumber() const;
+	FString GetClientName() const;
+
+	void Update(const TSharedRef<FWorkspaceUpdateContext, ESPMode::ThreadSafe>& Context);
+	void CancelUpdate();
+
+private:
+	static const TCHAR* DefaultBuildTargets[];
+	static const FWorkspaceSyncCategory DefaultSyncCategories[];
+	static const TCHAR* BuildVersionFileName;
+	static const TCHAR* VersionHeaderFileName;
+	static const TCHAR* ObjectVersionFileName;
+
+	int CurrentChangeNumber;
+	int PendingChangeNumber;
+	int LastBuiltChangeNumber;
+	TSharedRef<FLineBasedTextWriter> Log;
+	bool bSyncing;
+	mutable FCriticalSection CriticalSection;
+	TSharedRef<FCustomConfigFile, ESPMode::ThreadSafe> ProjectConfigFile;
+	TArray<FString> ProjectStreamFilter;
+	TSharedPtr<FWorkspaceUpdateContext, ESPMode::ThreadSafe> WorkerThreadContext;
+	FEvent* AbortEvent;
+	FRunnableThread* WorkerThread;
+	FProgressValue Progress;
+
+	static FWorkspace* ActiveWorkspace;
+
+	virtual uint32 Run() override;
+	void UpdateWorkspace(FWorkspaceUpdateContext& Context);
+	EWorkspaceUpdateResult UpdateWorkspaceInternal(FWorkspaceUpdateContext& Context, FString& OutStatusMessage);
+	bool MatchFilter(const FString& FileName, const FFileFilter& Filter);
+
+	static TArray<FString> GetSyncPaths(const FString& ClientRootPath, const FString& SelectedClientFileName);
+	static TSharedRef<FCustomConfigFile, ESPMode::ThreadSafe> ReadProjectConfigFile(const FString& LocalRootPath, const FString& SelectedLocalFileName, FLineBasedTextWriter& Log);
+	static TArray<FString> ReadProjectStreamFilter(FPerforceConnection& Perforce, const FCustomConfigFile& ProjectConfigFile, FEvent* AbortEvent, FLineBasedTextWriter& Log);
+	static FString FormatTime(long Seconds);
+	
+	bool HasModifiedSourceFiles() const;
+	bool FindUnresolvedFiles(TArray<FPerforceFileRecord>& OutUnresolvedFiles) const;
+	void UpdateSyncProgress(const FPerforceFileRecord& Record, TSet<FString>& RemainingFiles, int NumFiles);
+	bool UpdateVersionFile(const TCHAR* LocalPath, const TMap<FString, FString>& VersionStrings, int ChangeNumber) const;
+	bool WriteVersionFile(const FPerforceWhereRecord& WhereRecord, const FString& NewText) const;
+	static bool UpdateVersionLine(FString& Line, const FString& Prefix, const FString& Suffix);
+	static bool ReadToken(const FString& Line, int32& LineIdx, FString &OutToken);
+};
