Fix copyright info

diff --git a/Engine/Source/Programs/SlateUGS/Private/UGSCore/BuildStep.cpp b/Engine/Source/Programs/SlateUGS/Private/UGSCore/BuildStep.cpp
--- a/Engine/Source/Programs/SlateUGS/Private/UGSCore/BuildStep.cpp
+++ b/Engine/Source/Programs/SlateUGS/Private/UGSCore/BuildStep.cpp
@@ -1,4 +1,4 @@
-// Copyright 1998-2017 Epic Games, Inc. All Rights Reserved.
+// Copyright Epic Games, Inc. All Rights Reserved.
 
 #include "BuildStep.h"
 
diff --git a/Engine/Source/Programs/SlateUGS/Private/UGSCore/BuildStep.h b/Engine/Source/Programs/SlateUGS/Private/UGSCore/BuildStep.h
--- a/Engine/Source/Programs/SlateUGS/Private/UGSCore/BuildStep.h
+++ b/Engine/Source/Programs/SlateUGS/Private/UGSCore/BuildStep.h
@@ -1,4 +1,4 @@
-// Copyright 1998-2017 Epic Games, Inc. All Rights Reserved.
+// Copyright Epic Games, Inc. All Rights Reserved.
 
 #pragma once
 
diff --git a/Engine/Source/Programs/SlateUGS/Private/UGSCore/CustomConfigFile.cpp b/Engine/Source/Programs/SlateUGS/Private/UGSCore/CustomConfigFile.cpp
--- a/Engine/Source/Programs/SlateUGS/Private/UGSCore/CustomConfigFile.cpp
+++ b/Engine/Source/Programs/SlateUGS/Private/UGSCore/CustomConfigFile.cpp
@@ -1,4 +1,4 @@
-// Copyright 1998-2017 Epic Games, Inc. All Rights Reserved.
+// Copyright Epic Games, Inc. All Rights Reserved.
 
 #include "CustomConfigFile.h"
 #include "Misc/FileHelper.h"
diff --git a/Engine/Source/Programs/SlateUGS/Private/UGSCore/CustomConfigFile.h b/Engine/Source/Programs/SlateUGS/Private/UGSCore/CustomConfigFile.h
--- a/Engine/Source/Programs/SlateUGS/Private/UGSCore/CustomConfigFile.h
+++ b/Engine/Source/Programs/SlateUGS/Private/UGSCore/CustomConfigFile.h
@@ -1,4 +1,4 @@
-// Copyright 1998-2017 Epic Games, Inc. All Rights Reserved.
+// Copyright Epic Games, Inc. All Rights Reserved.
 
 #pragma once
 
diff --git a/Engine/Source/Programs/SlateUGS/Private/UGSCore/DetectProjectSettingsTask.cpp b/Engine/Source/Programs/SlateUGS/Private/UGSCore/DetectProjectSettingsTask.cpp
--- a/Engine/Source/Programs/SlateUGS/Private/UGSCore/DetectProjectSettingsTask.cpp
+++ b/Engine/Source/Programs/SlateUGS/Private/UGSCore/DetectProjectSettingsTask.cpp
@@ -1,4 +1,4 @@
-// Copyright 1998-2017 Epic Games, Inc. All Rights Reserved.
+// Copyright Epic Games, Inc. All Rights Reserved.
 
 #include "DetectProjectSettingsTask.h"
 #include "OutputAdapters.h"
diff --git a/Engine/Source/Programs/SlateUGS/Private/UGSCore/DetectProjectSettingsTask.h b/Engine/Source/Programs/SlateUGS/Private/UGSCore/DetectProjectSettingsTask.h
--- a/Engine/Source/Programs/SlateUGS/Private/UGSCore/DetectProjectSettingsTask.h
+++ b/Engine/Source/Programs/SlateUGS/Private/UGSCore/DetectProjectSettingsTask.h
@@ -1,4 +1,4 @@
-// Copyright 1998-2017 Epic Games, Inc. All Rights Reserved.
+// Copyright Epic Games, Inc. All Rights Reserved.
 
 #pragma once
 
diff --git a/Engine/Source/Programs/SlateUGS/Private/UGSCore/FileFilter.cpp b/Engine/Source/Programs/SlateUGS/Private/UGSCore/FileFilter.cpp
--- a/Engine/Source/Programs/SlateUGS/Private/UGSCore/FileFilter.cpp
+++ b/Engine/Source/Programs/SlateUGS/Private/UGSCore/FileFilter.cpp
@@ -1,4 +1,4 @@
-// Copyright 1998-2017 Epic Games, Inc. All Rights Reserved.
+// Copyright Epic Games, Inc. All Rights Reserved.
 
 #include "FileFilter.h"
 
diff --git a/Engine/Source/Programs/SlateUGS/Private/UGSCore/FileFilter.h b/Engine/Source/Programs/SlateUGS/Private/UGSCore/FileFilter.h
--- a/Engine/Source/Programs/SlateUGS/Private/UGSCore/FileFilter.h
+++ b/Engine/Source/Programs/SlateUGS/Private/UGSCore/FileFilter.h
@@ -1,4 +1,4 @@
-// Copyright 1998-2017 Epic Games, Inc. All Rights Reserved.
+// Copyright Epic Games, Inc. All Rights Reserved.
 
 #pragma once
 
diff --git a/Engine/Source/Programs/SlateUGS/Private/UGSCore/ModalTask.cpp b/Engine/Source/Programs/SlateUGS/Private/UGSCore/ModalTask.cpp
--- a/Engine/Source/Programs/SlateUGS/Private/UGSCore/ModalTask.cpp
+++ b/Engine/Source/Programs/SlateUGS/Private/UGSCore/ModalTask.cpp
@@ -1,4 +1,4 @@
-// Copyright 1998-2017 Epic Games, Inc. All Rights Reserved.
+// Copyright Epic Games, Inc. All Rights Reserved.
 
 #include "ModalTask.h"
 
diff --git a/Engine/Source/Programs/SlateUGS/Private/UGSCore/ModalTask.h b/Engine/Source/Programs/SlateUGS/Private/UGSCore/ModalTask.h
--- a/Engine/Source/Programs/SlateUGS/Private/UGSCore/ModalTask.h
+++ b/Engine/Source/Programs/SlateUGS/Private/UGSCore/ModalTask.h
@@ -1,4 +1,4 @@
-// Copyright 1998-2017 Epic Games, Inc. All Rights Reserved.
+// Copyright Epic Games, Inc. All Rights Reserved.
 
 #pragma once
 
diff --git a/Engine/Source/Programs/SlateUGS/Private/UGSCore/OutputAdapters.cpp b/Engine/Source/Programs/SlateUGS/Private/UGSCore/OutputAdapters.cpp
--- a/Engine/Source/Programs/SlateUGS/Private/UGSCore/OutputAdapters.cpp
+++ b/Engine/Source/Programs/SlateUGS/Private/UGSCore/OutputAdapters.cpp
@@ -1,4 +1,4 @@
-// Copyright 1998-2017 Epic Games, Inc. All Rights Reserved.
+// Copyright Epic Games, Inc. All Rights Reserved.
 
 #include "OutputAdapters.h"
 #include "Utility.h"
diff --git a/Engine/Source/Programs/SlateUGS/Private/UGSCore/OutputAdapters.h b/Engine/Source/Programs/SlateUGS/Private/UGSCore/OutputAdapters.h
--- a/Engine/Source/Programs/SlateUGS/Private/UGSCore/OutputAdapters.h
+++ b/Engine/Source/Programs/SlateUGS/Private/UGSCore/OutputAdapters.h
@@ -1,4 +1,4 @@
-// Copyright 1998-2017 Epic Games, Inc. All Rights Reserved.
+// Copyright Epic Games, Inc. All Rights Reserved.
 
 #pragma once
 
diff --git a/Engine/Source/Programs/SlateUGS/Private/UGSCore/Perforce.cpp b/Engine/Source/Programs/SlateUGS/Private/UGSCore/Perforce.cpp
--- a/Engine/Source/Programs/SlateUGS/Private/UGSCore/Perforce.cpp
+++ b/Engine/Source/Programs/SlateUGS/Private/UGSCore/Perforce.cpp
@@ -1,4 +1,4 @@
-// Copyright 1998-2017 Epic Games, Inc. All Rights Reserved.
+// Copyright Epic Games, Inc. All Rights Reserved.
 
 #include "Perforce.h"
 #include "Misc/FileHelper.h"
diff --git a/Engine/Source/Programs/SlateUGS/Private/UGSCore/Perforce.h b/Engine/Source/Programs/SlateUGS/Private/UGSCore/Perforce.h
--- a/Engine/Source/Programs/SlateUGS/Private/UGSCore/Perforce.h
+++ b/Engine/Source/Programs/SlateUGS/Private/UGSCore/Perforce.h
@@ -1,4 +1,4 @@
-// Copyright 1998-2017 Epic Games, Inc. All Rights Reserved.
+// Copyright Epic Games, Inc. All Rights Reserved.
 
 #pragma once
 
diff --git a/Engine/Source/Programs/SlateUGS/Private/UGSCore/Telemetry.cpp b/Engine/Source/Programs/SlateUGS/Private/UGSCore/Telemetry.cpp
--- a/Engine/Source/Programs/SlateUGS/Private/UGSCore/Telemetry.cpp
+++ b/Engine/Source/Programs/SlateUGS/Private/UGSCore/Telemetry.cpp
@@ -1,4 +1,4 @@
-// Copyright 1998-2017 Epic Games, Inc. All Rights Reserved.
+// Copyright Epic Games, Inc. All Rights Reserved.
 
 #include "Telemetry.h"
 
diff --git a/Engine/Source/Programs/SlateUGS/Private/UGSCore/Telemetry.h b/Engine/Source/Programs/SlateUGS/Private/UGSCore/Telemetry.h
--- a/Engine/Source/Programs/SlateUGS/Private/UGSCore/Telemetry.h
+++ b/Engine/Source/Programs/SlateUGS/Private/UGSCore/Telemetry.h
@@ -1,4 +1,4 @@
-// Copyright 1998-2017 Epic Games, Inc. All Rights Reserved.
+// Copyright Epic Games, Inc. All Rights Reserved.
 
 #pragma once
 
@@ -27,4 +27,4 @@ class FTelemetryWriter
 public:
 	FTelemetryWriter(const FString& InSqlConnectionString, const FString& InLogFileName);
 	~FTelemetryWriter();
-};
\ No newline at end of file
+};
diff --git a/Engine/Source/Programs/SlateUGS/Private/UGSCore/UserSettings.cpp b/Engine/Source/Programs/SlateUGS/Private/UGSCore/UserSettings.cpp
--- a/Engine/Source/Programs/SlateUGS/Private/UGSCore/UserSettings.cpp
+++ b/Engine/Source/Programs/SlateUGS/Private/UGSCore/UserSettings.cpp
@@ -1,4 +1,4 @@
-// Copyright 1998-2017 Epic Games, Inc. All Rights Reserved.
+// Copyright Epic Games, Inc. All Rights Reserved.
 
 #include "UserSettings.h"
 #include "HAL/FileManager.h"
diff --git a/Engine/Source/Programs/SlateUGS/Private/UGSCore/UserSettings.h b/Engine/Source/Programs/SlateUGS/Private/UGSCore/UserSettings.h
--- a/Engine/Source/Programs/SlateUGS/Private/UGSCore/UserSettings.h
+++ b/Engine/Source/Programs/SlateUGS/Private/UGSCore/UserSettings.h
@@ -1,4 +1,4 @@
-// Copyright 1998-2017 Epic Games, Inc. All Rights Reserved.
+// Copyright Epic Games, Inc. All Rights Reserved.
 
 #pragma once
 
diff --git a/Engine/Source/Programs/SlateUGS/Private/UGSCore/Utility.cpp b/Engine/Source/Programs/SlateUGS/Private/UGSCore/Utility.cpp
--- a/Engine/Source/Programs/SlateUGS/Private/UGSCore/Utility.cpp
+++ b/Engine/Source/Programs/SlateUGS/Private/UGSCore/Utility.cpp
@@ -1,4 +1,4 @@
-// Copyright 1998-2017 Epic Games, Inc. All Rights Reserved.
+// Copyright Epic Games, Inc. All Rights Reserved.
 
 #include "Utility.h"
 #include "HAL/FileManager.h"
diff --git a/Engine/Source/Programs/SlateUGS/Private/UGSCore/Utility.h b/Engine/Source/Programs/SlateUGS/Private/UGSCore/Utility.h
--- a/Engine/Source/Programs/SlateUGS/Private/UGSCore/Utility.h
+++ b/Engine/Source/Programs/SlateUGS/Private/UGSCore/Utility.h
@@ -1,4 +1,4 @@
-// Copyright 1998-2017 Epic Games, Inc. All Rights Reserved.
+// Copyright Epic Games, Inc. All Rights Reserved.
 
 #pragma once
 
diff --git a/Engine/Source/Programs/SlateUGS/Private/UGSCore/Workspace.cpp b/Engine/Source/Programs/SlateUGS/Private/UGSCore/Workspace.cpp
--- a/Engine/Source/Programs/SlateUGS/Private/UGSCore/Workspace.cpp
+++ b/Engine/Source/Programs/SlateUGS/Private/UGSCore/Workspace.cpp
@@ -1,4 +1,4 @@
-// Copyright 1998-2017 Epic Games, Inc. All Rights Reserved.
+// Copyright Epic Games, Inc. All Rights Reserved.
 
 #include "Workspace.h"
 #include "Misc/Paths.h"
diff --git a/Engine/Source/Programs/SlateUGS/Private/UGSCore/Workspace.h b/Engine/Source/Programs/SlateUGS/Private/UGSCore/Workspace.h
--- a/Engine/Source/Programs/SlateUGS/Private/UGSCore/Workspace.h
+++ b/Engine/Source/Programs/SlateUGS/Private/UGSCore/Workspace.h
@@ -1,4 +1,4 @@
-// Copyright 1998-2017 Epic Games, Inc. All Rights Reserved.
+// Copyright Epic Games, Inc. All Rights Reserved.
 
 #pragma once
 
