Fix compiler error, implicit int -> bool conversion

diff --git a/Engine/Source/Programs/SlateUGS/Private/SlateUGSApp.cpp b/Engine/Source/Programs/SlateUGS/Private/SlateUGSApp.cpp
--- a/Engine/Source/Programs/SlateUGS/Private/SlateUGSApp.cpp
+++ b/Engine/Source/Programs/SlateUGS/Private/SlateUGSApp.cpp
@@ -19,7 +19,7 @@
 		for (int i = 0; i < 35; i++)
 		{
 			TSharedPtr<HordeBuildRowInfo> Row = MakeShared<HordeBuildRowInfo>();
-			Row->bBuildStatus = i % 2;
+			Row->bBuildStatus = !!(i % 2);
 			Row->Changelist   = FText::FromString("12345678");
 			Row->Time         = FText::FromString("11:48 AM");
 			Row->Author       = FText::FromString("Robert Seiver");
