Fix not saving out the user settings
Fix an assert when trying to save settings

diff --git a/Engine/Source/Programs/SlateUGS/Private/UGSCore/CustomConfigFile.cpp b/Engine/Source/Programs/SlateUGS/Private/UGSCore/CustomConfigFile.cpp
--- a/Engine/Source/Programs/SlateUGS/Private/UGSCore/CustomConfigFile.cpp
+++ b/Engine/Source/Programs/SlateUGS/Private/UGSCore/CustomConfigFile.cpp
@@ -581,8 +581,6 @@ void FCustomConfigFile::Parse(const TArray<FString>& Lines)
 
 void FCustomConfigFile::Save(const TCHAR* FileName)
 {
-	return;
-
 	TArray<FString> Lines;
 	for(int Idx = 0; Idx < Sections.Num(); Idx++)
 	{
@@ -609,6 +607,8 @@ void FCustomConfigFile::Save(const TCHAR* FileName)
 			Lines.Add(FString());
 		}
 	}
+
+	FFileHelper::SaveStringArrayToFile(Lines, FileName);
 }
 
 TSharedPtr<FCustomConfigSection> FCustomConfigFile::FindSection(const TCHAR* Name)
diff --git a/Engine/Source/Programs/SlateUGS/Private/UGSCore/UserSettings.cpp b/Engine/Source/Programs/SlateUGS/Private/UGSCore/UserSettings.cpp
--- a/Engine/Source/Programs/SlateUGS/Private/UGSCore/UserSettings.cpp
+++ b/Engine/Source/Programs/SlateUGS/Private/UGSCore/UserSettings.cpp
@@ -170,6 +170,9 @@ FUserSettings::FUserSettings(const FString& InFileName)
 	ConfigFile.TryGetValues(TEXT("General.OtherProjectFileNames"), OtherProjectFileNames);
 	ConfigFile.TryGetValues(TEXT("General.SyncFilter"), SyncView);
 	ConfigFile.TryGetValues(TEXT("General.SyncExcludedCategories"), SyncExcludedCategories);
+
+	// TODO i set this default value as we assert when saving if no value is set but honestly idk what Any means
+	SyncType = ELatestChangeType::Any;
 	if(TryParse(ConfigFile.GetValue(TEXT("General.SyncType"), TEXT("")), SyncType))
 	{
 		SyncType = ELatestChangeType::Good;
