Introduced UGSTab and UGSTabManager which acts as a buffer between UGSCore code and Widget UI (Slate) code.

diff --git a/Engine/Source/Programs/SlateUGS/Private/SlateUGSApp.cpp b/Engine/Source/Programs/SlateUGS/Private/SlateUGSApp.cpp
--- a/Engine/Source/Programs/SlateUGS/Private/SlateUGSApp.cpp
+++ b/Engine/Source/Programs/SlateUGS/Private/SlateUGSApp.cpp
@@ -4,70 +4,17 @@
 #include "RequiredProgramMainCPPInclude.h"
 #include "Framework/Application/SlateApplication.h"
 
-#include "Widgets/Docking/SDockTab.h"
-
+#include "UGSCore/GameSyncController.h"
 #include "Widgets/SGameSyncTab.h"
 #include "Widgets/SWorkspaceWindow.h"
 #include "Widgets/SEmptyTab.h"
 
+#include "UGSTabManager.h"
 
 IMPLEMENT_APPLICATION(SlateUGS, "SlateUGS");
 
 #define LOCTEXT_NAMESPACE "SlateUGS"
 
-namespace
-{
-	TSharedRef<SDockTab> SpawnEmptyTab(const FSpawnTabArgs& Arguments)
-	{
-		return SNew(SDockTab).TabRole(ETabRole::MajorTab)
-			[
-				SNew(SEmptyTab)
-			];
-	}
-	TSharedRef<SDockTab> SpawnActiveTab(const FSpawnTabArgs& Arguments)
-	{
-		// Todo: replace with real Horde build data (or gather this in the SGameSyncTab construct)
-		TArray<TSharedPtr<HordeBuildRowInfo>> HordeBuilds;
-		for (int i = 0; i < 35; i++)
-		{
-			TSharedPtr<HordeBuildRowInfo> Row = MakeShared<HordeBuildRowInfo>();
-			Row->bBuildStatus = !!(i % 2);
-			Row->Changelist   = FText::FromString("12345678");
-			Row->Time         = FText::FromString("11:48 AM");
-			Row->Author       = FText::FromString("Robert Seiver");
-			Row->Description  = FText::FromString("Fixed the thing");
-			Row->Status       = FText::FromString("Used by Brandon Schaefer, Michael Sartain, ...");
-
-			HordeBuilds.Add(Row);
-		}
-
-		return SNew(SDockTab).TabRole(ETabRole::MajorTab)
-			[
-				SNew(SGameSyncTab)
-					.HordeBuilds(HordeBuilds)
-			];
-	}
-	
-	void BuildWindow()
-	{
-		FGlobalTabmanager::Get()->RegisterTabSpawner("EmptyTab", FOnSpawnTab::CreateStatic(SpawnEmptyTab));
-		FGlobalTabmanager::Get()->RegisterTabSpawner("ActiveTab", FOnSpawnTab::CreateStatic(SpawnActiveTab));
-		TSharedRef<FTabManager::FLayout> Layout = FTabManager::NewLayout("UGS_Layout")
-		->AddArea
-		(
-			FTabManager::NewArea(1230, 900)
-			->Split
-			(
-				FTabManager::NewStack()
-				->AddTab("EmptyTab", ETabState::OpenedTab)
-				->AddTab("ActiveTab", ETabState::ClosedTab) // Todo: seems to be only one tab per ID, need several tabs all with UGSActiveTab contents
-				->SetForegroundTab(FName("EmptyTab"))
-			)
-		);
-		FGlobalTabmanager::Get()->RestoreFrom(Layout, TSharedPtr<SWindow>());
-	}
-}
-
 int RunSlateUGS(const TCHAR* CommandLine)
 {
 	FTaskTagScope TaskTagScope(ETaskTag::EGameThread);	
@@ -92,7 +39,8 @@ int RunSlateUGS(const TCHAR* CommandLine)
 	FAppStyle::SetAppStyleSetName(FAppStyle::GetAppStyleSetName());
 
 	// Build the slate UI for the program window
-	BuildWindow();
+	UGSTabManager TabManager; 
+	TabManager.ConstructTabs();
 
 	// loop while the server does the rest
 	while (!IsEngineExitRequested())
diff --git a/Engine/Source/Programs/SlateUGS/Private/UGSCore/GameSyncController.cpp b/Engine/Source/Programs/SlateUGS/Private/UGSCore/GameSyncController.cpp
new file mode 100644
--- /dev/null
+++ b/Engine/Source/Programs/SlateUGS/Private/UGSCore/GameSyncController.cpp
@@ -0,0 +1,178 @@
+#include "GameSyncController.h"
+
+#include "BuildStep.h"
+#include "DetectProjectSettingsTask.h"
+#include "Utility.h"
+#include "UserSettings.h"
+
+namespace
+{
+// TODO super hacky... fix up later
+#if PLATFORM_WINDOWS
+	const TCHAR* HostPlatform = TEXT("Win64");
+#elif PLATFORM_MAC
+	const TCHAR* HostPlatform = TEXT("Mac");
+#else
+	const TCHAR* HostPlatform = TEXT("Linux");
+#endif
+}
+
+class FLineWriter : public FLineBasedTextWriter
+{
+public:
+	virtual void FlushLine(const FString& Line) override
+	{
+		printf("%s\n", TCHAR_TO_ANSI(*Line));
+	}
+};
+
+// TODO move most of these to a class with access to the settings
+//bool ShouldSyncPrecompiledEditor(TSharedPtr<FUserSettings> Settings)
+bool ShouldSyncPrecompiledEditor()
+{
+	return false;
+	//return Settings->bSyncPrecompiledEditor;// && PerforceMonitor->HasZippedBinaries();
+}
+
+// Honestly ... seems ... super hacky/hardcoded. With out all these you assert when trying to merge build targets sooo a bit odd
+// TODO Need to do each of these ... per ... platform??
+TMap<FGuid, FCustomConfigObject> GetDefaultBuildStepObjects(const FString& EditorTargetName, TSharedPtr<FUserSettings> Settings)
+{
+	TArray<FBuildStep> DefaultBuildSteps;
+	DefaultBuildSteps.Add(FBuildStep(FGuid(0x01F66060, 0x73FA4CC8, 0x9CB3E217, 0xFBBA954E), 0, TEXT("Compile UnrealHeaderTool"), TEXT("Compiling UnrealHeaderTool..."), 1, TEXT("UnrealHeaderTool"), HostPlatform, TEXT("Development"), TEXT(""), !ShouldSyncPrecompiledEditor()));
+	FString ActualEditorTargetName = (EditorTargetName.Len() > 0)? EditorTargetName : "UnrealEditor";
+	DefaultBuildSteps.Add(FBuildStep(FGuid(0xF097FF61, 0xC9164058, 0x839135B4, 0x6C3173D5), 1, FString::Printf(TEXT("Compile %s"), *ActualEditorTargetName), FString::Printf(TEXT("Compiling %s..."), *ActualEditorTargetName), 10, ActualEditorTargetName, HostPlatform, ::ToString(Settings->CompiledEditorBuildConfig), TEXT(""), !ShouldSyncPrecompiledEditor()));
+	DefaultBuildSteps.Add(FBuildStep(FGuid(0xC6E633A1, 0x956F4AD3, 0xBC956D06, 0xD131E7B4), 2, TEXT("Compile ShaderCompileWorker"), TEXT("Compiling ShaderCompileWorker..."), 1, TEXT("ShaderCompileWorker"), HostPlatform, TEXT("Development"), TEXT(""), !ShouldSyncPrecompiledEditor()));
+	DefaultBuildSteps.Add(FBuildStep(FGuid(0x24FFD88C, 0x79014899, 0x9696AE10, 0x66B4B6E8), 3, TEXT("Compile UnrealLightmass"), TEXT("Compiling UnrealLightmass..."), 1, TEXT("UnrealLightmass"), HostPlatform, TEXT("Development"), TEXT(""), !ShouldSyncPrecompiledEditor()));
+	DefaultBuildSteps.Add(FBuildStep(FGuid(0xFFF20379, 0x06BF4205, 0x8A3EC534, 0x27736688), 4, TEXT("Compile CrashReportClient"), TEXT("Compiling CrashReportClient..."), 1, TEXT("CrashReportClient"), HostPlatform, TEXT("Development"), TEXT(""), !ShouldSyncPrecompiledEditor()));
+	DefaultBuildSteps.Add(FBuildStep(FGuid(0x89FE8A79, 0xD2594C7B, 0xBFB468F7, 0x218B91C2), 5, TEXT("Compile UnrealInsights"), TEXT("Compiling UnrealInsights..."), 1, TEXT("UnrealInsights"), HostPlatform, TEXT("Development"), TEXT(""), !ShouldSyncPrecompiledEditor()));
+	DefaultBuildSteps.Add(FBuildStep(FGuid(0x46312669, 0x5069428D, 0x8D72C241, 0x6C5A322E), 6, TEXT("Launch UnrealInsights"), TEXT("Running UnrealInsights..."), 1, TEXT("UnrealInsights"), HostPlatform, TEXT("Shipping"), TEXT(""), !ShouldSyncPrecompiledEditor()));
+	DefaultBuildSteps.Add(FBuildStep(FGuid(0xBB48CA5B, 0x56824432, 0x824DC451, 0x336A6523), 7, TEXT("Compile Zen Dashboard"), TEXT("Compile ZenDashboard Step..."), 1, TEXT("ZenDashboard"), HostPlatform, TEXT("Development"), TEXT(""), !ShouldSyncPrecompiledEditor()));
+	DefaultBuildSteps.Add(FBuildStep(FGuid(0x586CC0D3, 0x39144DF9, 0xACB62C02, 0xCD9D4FC6), 8, TEXT("Launch Zen Dashboard"), TEXT("Running Zen Dashboard..."), 1, TEXT("ZenDashboard"), HostPlatform, TEXT("Development"), TEXT(""), !ShouldSyncPrecompiledEditor()));
+	DefaultBuildSteps.Add(FBuildStep(FGuid(0x91C2A429, 0xC39149B4, 0x92A54E6B, 0xE71E0F00), 9, TEXT("Compile SwitchboardListener"), TEXT("Compiling SwitchboardListener..."), 1, TEXT("SwitchboardListener"), HostPlatform, TEXT("Development"), TEXT(""), !ShouldSyncPrecompiledEditor()));
+	DefaultBuildSteps.Add(FBuildStep(FGuid(0x5036C75B, 0x8DF04329, 0x82A1869D, 0xD2D48605), 10, TEXT("Compile UnrealMultiUserServer"), TEXT("Compiling UnrealMultiUserServer..."), 1, TEXT("UnrealMultiUserServer"), HostPlatform, TEXT("Development"), TEXT(""), !ShouldSyncPrecompiledEditor()));
+	DefaultBuildSteps.Add(FBuildStep(FGuid(0x274B89C3, 0x9DC64465, 0xA50840AB, 0xC4593CC2), 11, TEXT("Compile UnrealMultiUserSlateServer"), TEXT("Compiling UnrealMultiUserSlateServer..."), 1, TEXT("UnrealMultiUserSlateServer"), HostPlatform, TEXT("Development"), TEXT(""), !ShouldSyncPrecompiledEditor()));
+
+	TMap<FGuid,FCustomConfigObject> DefaultBuildStepObjects;
+	for(const FBuildStep& DefaultBuildStep : DefaultBuildSteps)
+	{
+		DefaultBuildStepObjects.Add(DefaultBuildStep.UniqueId, DefaultBuildStep.ToConfigObject());
+	}
+	return DefaultBuildStepObjects;
+}
+
+EBuildConfig GetEditorBuildConfig()
+{
+	return EBuildConfig::Development;
+	//return ShouldSyncPrecompiledEditor()? EBuildConfig::Development : Settings->CompiledEditorBuildConfig;
+}
+
+FString GetEditorExePath(EBuildConfig Config, TSharedPtr<FDetectProjectSettingsTask> DetectSettings)
+{
+	FString ExeFileName = TEXT("UnrealEditor");
+
+	if(Config != EBuildConfig::DebugGame && Config != EBuildConfig::Development)
+	{
+#if PLATFORM_WINDOWS
+		ExeFileName = FString::Printf(TEXT("UnrealEditor-%s-%s.exe"), HostPlatform, *::ToString(Config));
+#else
+		ExeFileName = FString::Printf(TEXT("UnrealEditor-%s-%s"), HostPlatform, *::ToString(Config));
+#endif
+	}
+	return DetectSettings->BranchDirectoryName / "Engine" / "Binaries" / HostPlatform / ExeFileName;
+}
+
+TMap<FString, FString> GetWorkspaceVariables(TSharedPtr<FDetectProjectSettingsTask> DetectSettings)
+{
+	EBuildConfig EditorBuildConfig = GetEditorBuildConfig();
+
+	TMap<FString, FString> Variables;
+	Variables.Add("BranchDir", DetectSettings->BranchDirectoryName);
+	Variables.Add("ProjectDir", FPaths::GetPath(DetectSettings->NewSelectedFileName));
+	Variables.Add("ProjectFile", DetectSettings->NewSelectedFileName);
+
+	// Todo: These might not be called "UE4*" anymore
+	Variables.Add("UE4EditorExe", GetEditorExePath(EditorBuildConfig, DetectSettings));
+	Variables.Add("UE4EditorCmdExe", GetEditorExePath(EditorBuildConfig, DetectSettings).Replace(TEXT(".exe"), TEXT("-Cmd.exe")));
+	Variables.Add("UE4EditorConfig", ::ToString(EditorBuildConfig));
+	Variables.Add("UE4EditorDebugArg", (EditorBuildConfig == EBuildConfig::Debug || EditorBuildConfig == EBuildConfig::DebugGame)? " -debug" : "");
+	return Variables;
+} 
+
+void GameSyncController::SetupWorkspace(FString ProjectFileName)
+{
+	printf("ProjectFileName is: %s\n", TCHAR_TO_ANSI(*ProjectFileName));
+	ProjectFileName = FUtility::GetPathWithCorrectCase(ProjectFileName);
+	TSharedPtr<FDetectProjectSettingsTask> DetectSettings = MakeShared<FDetectProjectSettingsTask>(MakeShared<FPerforceConnection>(TEXT(""), TEXT(""), TEXT("")), ProjectFileName, MakeShared<FLineWriter>());
+
+	// AbortEvent for canceling but honestly not setup well.
+	// Todo: fix this up and likely enable exceptions
+	FEvent* AbortEvent = nullptr;
+	DetectSettings->Run(AbortEvent);
+
+	FString DataFolder = FString(FPlatformProcess::UserSettingsDir()) / TEXT("UnrealGameSync");
+	IFileManager::Get().MakeDirectory(*DataFolder);
+
+	TSharedPtr<FUserSettings> Settings = MakeShared<FUserSettings>(*(DataFolder / TEXT("UnrealGameSync.ini")));
+
+	TSharedRef<FPerforceConnection> PerforceClient = DetectSettings->PerforceClient.ToSharedRef();
+	TSharedRef<FUserWorkspaceSettings> WorkspaceSettings = Settings->FindOrAddWorkspace(*DetectSettings->BranchClientPath);
+	TSharedRef<FUserProjectSettings> ProjectSettings = Settings->FindOrAddProject(*DetectSettings->NewSelectedClientFileName);
+
+	FString SelectedClientFileName = DetectSettings->NewSelectedClientFileName;
+	FString BranchDirectoryName = DetectSettings->BranchDirectoryName;
+	FString BranchClientPath = DetectSettings->BranchClientPath;
+	FString SelectedProjectIdentifier = DetectSettings->NewSelectedProjectIdentifier;
+	FString EditorTarget = DetectSettings->NewProjectEditorTarget;
+
+	// Check if the project we've got open in this workspace is the one we're actually synced to
+	int CurrentChangeNumber = -1;
+	if(WorkspaceSettings->CurrentProjectIdentifier == SelectedProjectIdentifier)
+	{
+		CurrentChangeNumber = WorkspaceSettings->CurrentChangeNumber;
+	}
+
+	FString ClientKey = BranchClientPath.Replace(*FString::Printf(TEXT("//%s/"), *PerforceClient->ClientName), TEXT(""));
+	if(ClientKey.EndsWith(TEXT("/")))
+	{
+		ClientKey = ClientKey.Left(ClientKey.Len() - 1);
+	}
+
+	FString ProjectLogBaseName = DataFolder / FString::Printf(TEXT("%s@%s"), *PerforceClient->ClientName, *ClientKey.Replace(TEXT("/"), TEXT("$")));
+	FString TelemetryProjectIdentifier = FPerforceUtils::GetClientOrDepotDirectoryName(*SelectedProjectIdentifier);
+
+	Workspace = MakeShared<FWorkspace>(PerforceClient, BranchDirectoryName, ProjectFileName, BranchClientPath, SelectedClientFileName, CurrentChangeNumber, WorkspaceSettings->LastBuiltChangeNumber, TelemetryProjectIdentifier, MakeShared<FLineWriter>());
+
+	// Todo:Move into a sync operation
+	// TArray<FString> CombinedSyncFilter = FUserSettings::GetCombinedSyncFilter(Workspace->GetSyncCategories(), Settings->SyncView, Settings->SyncExcludedCategories, WorkspaceSettings->SyncView, WorkspaceSettings->SyncExcludedCategories);
+
+	// // Options on what to do with workspace when updating it
+	// EWorkspaceUpdateOptions Options = EWorkspaceUpdateOptions::Sync | EWorkspaceUpdateOptions::SyncArchives | EWorkspaceUpdateOptions::GenerateProjectFiles;
+	// if(Settings->bAutoResolveConflicts)
+	// {
+	// 	Options |= EWorkspaceUpdateOptions::AutoResolveChanges;
+	// }
+	// if(Settings->bUseIncrementalBuilds)
+	// {
+	// 	Options |= EWorkspaceUpdateOptions::UseIncrementalBuilds;
+	// }
+	// if(Settings->bBuildAfterSync)
+	// {
+	// 	Options |= EWorkspaceUpdateOptions::Build;
+	// }
+	// if(Settings->bBuildAfterSync && Settings->bRunAfterSync)
+	// {
+	// 	Options |= EWorkspaceUpdateOptions::RunAfterSync; 
+	// }
+	// if(Settings->bOpenSolutionAfterSync)
+	// {
+	// 	Options |= EWorkspaceUpdateOptions::OpenSolutionAfterSync;
+	// }
+
+	// // Hacking in the CL
+	// int ChangeNumber = 20658981;
+	// TSharedRef<FWorkspaceUpdateContext, ESPMode::ThreadSafe> Context = MakeShared<FWorkspaceUpdateContext, ESPMode::ThreadSafe>(ChangeNumber, Options, CombinedSyncFilter, GetDefaultBuildStepObjects(EditorTarget, Settings), ProjectSettings->BuildSteps, TSet<FGuid>(), GetWorkspaceVariables(DetectSettings));
+
+	// // Update the workspace with the Context!
+	// Workspace->Update(Context);
+}
diff --git a/Engine/Source/Programs/SlateUGS/Private/UGSCore/GameSyncController.h b/Engine/Source/Programs/SlateUGS/Private/UGSCore/GameSyncController.h
new file mode 100644
--- /dev/null
+++ b/Engine/Source/Programs/SlateUGS/Private/UGSCore/GameSyncController.h
@@ -0,0 +1,79 @@
+// Copyright Epic Games, Inc. All Rights Reserved.
+
+#pragma once
+
+#include "CoreMinimal.h"
+#include "Workspace.h"
+
+/*
+#include "Perforce.h"
+#include "CustomConfigFile.h"
+#include "Misc/EnumClassFlags.h"
+
+enum class EWorkspaceUpdateOptions
+{
+	Sync = 0x01,
+	SyncSingleChange = 0x02,
+	AutoResolveChanges = 0x04,
+	GenerateProjectFiles = 0x08,
+	SyncArchives = 0x10,
+	Build = 0x20,
+	UseIncrementalBuilds = 0x40,
+	ScheduledBuild = 0x80,
+	RunAfterSync = 0x100,
+	OpenSolutionAfterSync = 0x200,
+	ContentOnly = 0x400
+};
+
+ENUM_CLASS_FLAGS(EWorkspaceUpdateOptions)
+
+enum class EWorkspaceUpdateResult
+{
+	Canceled,
+	FailedToSync,
+	FilesToResolve,
+	FilesToClobber,
+	FailedToCompile,
+	FailedToCompileWithCleanWorkspace,
+	Success,
+};
+
+FString ToString(EWorkspaceUpdateResult WorkspaceUpdateResult);
+bool TryParse(const TCHAR* Text, EWorkspaceUpdateResult& OutWorkspaceUpdateResult);
+
+struct FWorkspaceUpdateContext
+{
+	FDateTime StartTime;
+	int ChangeNumber;
+	EWorkspaceUpdateOptions Options;
+	TArray<FString> SyncFilter;
+	TMap<FString, FString> ArchiveTypeToDepotPath;
+	TMap<FString, bool> ClobberFiles;
+	TMap<FGuid,FCustomConfigObject> DefaultBuildSteps;
+	TArray<FCustomConfigObject> UserBuildStepObjects;
+	TSet<FGuid> CustomBuildSteps;
+	TMap<FString, FString> Variables;
+	FPerforceSyncOptions PerforceSyncOptions;
+
+	FWorkspaceUpdateContext(int InChangeNumber, EWorkspaceUpdateOptions InOptions, const TArray<FString>& InSyncFilter, const TMap<FGuid, FCustomConfigObject>& InDefaultBuildSteps, const TArray<FCustomConfigObject>& InUserBuildSteps, const TSet<FGuid>& InCustomBuildSteps, const TMap<FString, FString>& InVariables);
+};
+
+struct FWorkspaceSyncCategory
+{
+	FGuid UniqueId;
+	bool bEnable;
+	FString Name;
+	TArray<FString> Paths;
+
+	FWorkspaceSyncCategory(const FGuid& InUniqueId);
+	FWorkspaceSyncCategory(const FGuid& InUniqueId, const TCHAR* InName, const TCHAR* InPaths);
+};
+*/
+
+class GameSyncController
+{
+public:
+	void SetupWorkspace(FString ProjectFileName);
+private:
+	TSharedPtr<FWorkspace> Workspace;
+};
diff --git a/Engine/Source/Programs/SlateUGS/Private/UGSTab.cpp b/Engine/Source/Programs/SlateUGS/Private/UGSTab.cpp
new file mode 100644
--- /dev/null
+++ b/Engine/Source/Programs/SlateUGS/Private/UGSTab.cpp
@@ -0,0 +1,80 @@
+// Copyright Epic Games, Inc. All Rights Reserved.
+
+#include "UGSTab.h"
+
+#define LOCTEXT_NAMESPACE "UGSTab"
+
+UGSTab::UGSTab() : TabArgs(nullptr, FTabId()),
+				   TabWidget(SNew(SDockTab)),
+				   EmptyTabView(SNew(SEmptyTab).Tab(this)),
+				   GameSyncTabView(SNew(SGameSyncTab)),
+				   WorkspaceWindowView(SNew(SWorkspaceWindow).Tab(this))
+{
+	TabWidget->SetContent(EmptyTabView);
+}
+
+const TSharedRef<SDockTab> UGSTab::GetTabWidget()
+{
+	return TabWidget;
+}
+
+void UGSTab::SetTabArgs(FSpawnTabArgs InTabArgs) 
+{ 
+	TabArgs = InTabArgs; 
+}
+
+FSpawnTabArgs UGSTab::GetTabArgs() const 
+{ 
+	return TabArgs; 
+}
+
+FReply UGSTab::OnWorkspaceChosen(FString Path)
+{
+	fprintf(stderr, "Workspace path text received by OnWorkspaceChosen is: %s\n", TCHAR_TO_ANSI(*Path)); // Todo: Actually get the path variable bound
+	bool bIsDataValid = !Path.IsEmpty(); // Todo: Actually validate the data
+	if (bIsDataValid)
+	{
+		TabWidget->SetContent(GameSyncTabView); // Todo: Set GameSyncTabView data
+	}
+	else
+	{
+		// Todo: factor out into an error window widget class (check if one already exists)
+		TSharedRef<SWindow> Window = SNew(SWindow)
+		.Title(LOCTEXT("ErrorWindowTitle", "Error Opening Project"))
+		.SizingRule(ESizingRule::Autosized)
+		.MaxWidth(400.0f)
+		[
+			SNew(SVerticalBox)
+			+SVerticalBox::Slot()
+			.AutoHeight()
+			.HAlign(HAlign_Center)
+			.Padding(10.0f, 10.0f)
+			[
+				SNew(STextBlock)
+				.AutoWrapText(true)
+				.Text(LOCTEXT("ErrorText", "Error opening .uproject file, try again")) // Todo: detect the actual reason for error and report it here
+			]
+			+SVerticalBox::Slot()
+			.AutoHeight()
+			.HAlign(HAlign_Center)
+			.Padding(0.0f, 10.0f, 0.0f, 0.0f)
+			[
+				SNew(SButton)
+				.HAlign(HAlign_Center)
+				.Text(LOCTEXT("ErrorWindowOkayButtonText", "Ok"))
+				.OnClicked_Lambda([&Window]() 
+				{
+					Window->RequestDestroyWindow();
+					return FReply::Handled(); 
+				})
+			]
+		];
+
+		FSlateApplication::Get().AddModalWindow(Window, FSlateApplication::Get().GetActiveModalWindow(), false); // Todo: Figure out parent window for modal window
+		// Todo: loading screen widget for detect settings
+	}
+
+	return FReply::Handled();
+}
+
+#undef LOCTEXT_NAMESPACE
diff --git a/Engine/Source/Programs/SlateUGS/Private/UGSTab.h b/Engine/Source/Programs/SlateUGS/Private/UGSTab.h
new file mode 100644
--- /dev/null
+++ b/Engine/Source/Programs/SlateUGS/Private/UGSTab.h
@@ -0,0 +1,34 @@
+// Copyright Epic Games, Inc. All Rights Reserved.
+
+#pragma once
+
+#include "Widgets/SEmptyTab.h" 
+#include "Widgets/SGameSyncTab.h"
+#include "Widgets/SWorkspaceWindow.h"
+#include "UGSCore/GameSyncController.h"
+
+class UGSTab
+{
+public:
+	UGSTab();  // Todo: change all _Args to InArgs
+
+	const TSharedRef<SDockTab> GetTabWidget(); // Todo: check const-correctness of this (is the pointer const or is what it's pointing to const)
+
+	void SetTabArgs(FSpawnTabArgs InTabArgs);
+	FSpawnTabArgs GetTabArgs() const;
+
+	// Slate callbacks
+	FReply OnWorkspaceChosen(FString Path);
+private:
+
+	// Slate Data
+	FSpawnTabArgs TabArgs;
+	TSharedRef<SDockTab> TabWidget;
+
+	TSharedRef<SEmptyTab> EmptyTabView;
+	TSharedRef<SGameSyncTab> GameSyncTabView;
+	TSharedRef<SWorkspaceWindow> WorkspaceWindowView;
+
+	// Controller data
+	GameSyncController SyncController; 
+};
diff --git a/Engine/Source/Programs/SlateUGS/Private/UGSTabManager.cpp b/Engine/Source/Programs/SlateUGS/Private/UGSTabManager.cpp
new file mode 100644
--- /dev/null
+++ b/Engine/Source/Programs/SlateUGS/Private/UGSTabManager.cpp
@@ -0,0 +1,43 @@
+// Copyright Epic Games, Inc. All Rights Reserved.
+
+#include "UGSTabManager.h"
+#include "Containers/Array.h"
+
+void UGSTabManager::ConstructTabs()
+{
+	TSharedRef<FTabManager::FStack> TabStack = FTabManager::NewStack();
+
+	for (int TabIndex = 0; TabIndex < MaxTabs; TabIndex++)
+	{
+		const FName TabId = FName(FString(TEXT("UGS Tab: ")) + FString::FromInt(TabIndex));
+		FGlobalTabmanager::Get()->RegisterTabSpawner(TabId, 
+			FOnSpawnTab::CreateLambda([this, TabIndex] (const FSpawnTabArgs& Args) { return SpawnTab(TabIndex, Args); }));//,)
+		// 	FCanSpawnTab::CreateLambda([this] (const FSpawnTabArgs& Args) { }) // Todo: going to need to limit this to max tabs later
+		// );
+
+		// Leave the first tab opened, close the rest
+		if (TabIndex == 0)
+		{
+			TabStack->AddTab(TabId, ETabState::OpenedTab);
+		}
+		else
+		{
+			TabStack->AddTab(TabId, ETabState::ClosedTab);
+		}
+	}
+
+	TSharedRef<FTabManager::FLayout> Layout = FTabManager::NewLayout("UGS_Layout")
+	->AddArea
+	(
+		FTabManager::NewArea(1230, 900)
+		->Split(TabStack)
+	);
+
+	FGlobalTabmanager::Get()->RestoreFrom(Layout, nullptr);
+}
+
+TSharedRef<SDockTab> UGSTabManager::SpawnTab(int Index, const FSpawnTabArgs& Args)
+{
+	Tabs[Index].SetTabArgs(Args);
+	return Tabs[Index].GetTabWidget();
+}
diff --git a/Engine/Source/Programs/SlateUGS/Private/UGSTabManager.h b/Engine/Source/Programs/SlateUGS/Private/UGSTabManager.h
new file mode 100644
--- /dev/null
+++ b/Engine/Source/Programs/SlateUGS/Private/UGSTabManager.h
@@ -0,0 +1,18 @@
+// Copyright Epic Games, Inc. All Rights Reserved.
+
+#pragma once
+
+#include "UGSTab.h"
+
+class UGSTabManager
+{
+public:
+	UGSTabManager() = default;
+	void ConstructTabs();
+	TSharedRef<SDockTab> SpawnTab(int Index, const FSpawnTabArgs& Args);
+	bool CanSpawnTab() const;
+	int GetMaxTabs() const;
+private:
+	static constexpr int MaxTabs = 1;
+	TStaticArray<UGSTab, MaxTabs> Tabs;
+};
\ No newline at end of file
diff --git a/Engine/Source/Programs/SlateUGS/Private/Widgets/SEmptyTab.cpp b/Engine/Source/Programs/SlateUGS/Private/Widgets/SEmptyTab.cpp
--- a/Engine/Source/Programs/SlateUGS/Private/Widgets/SEmptyTab.cpp
+++ b/Engine/Source/Programs/SlateUGS/Private/Widgets/SEmptyTab.cpp
@@ -5,10 +5,14 @@
 #include "SWorkspaceWindow.h"
 #include "Styling/AppStyle.h"
 
+#include "UGSTab.h"
+
 #define LOCTEXT_NAMESPACE "UGSEmptyTab"
 
 void SEmptyTab::Construct(const FArguments& InArgs)
 {
+	Tab = InArgs._Tab;
+
 	this->ChildSlot
 	[
 		SNew(SBox)
@@ -71,7 +75,7 @@ FReply SEmptyTab::OnOpenProjectClicked()
 	];
 
 	FSlateApplication& SlateApplication = FSlateApplication::Get();
-	SlateApplication.AddModalWindow(Window, SlateApplication.FindBestParentWindowForDialogs(nullptr), false);
+	SlateApplication.AddModalWindow(Window, Tab->GetTabArgs().GetOwnerWindow(), false);
 
 	return FReply::Handled();
 }
diff --git a/Engine/Source/Programs/SlateUGS/Private/Widgets/SEmptyTab.h b/Engine/Source/Programs/SlateUGS/Private/Widgets/SEmptyTab.h
--- a/Engine/Source/Programs/SlateUGS/Private/Widgets/SEmptyTab.h
+++ b/Engine/Source/Programs/SlateUGS/Private/Widgets/SEmptyTab.h
@@ -4,20 +4,19 @@
 
 #include "Widgets/SCompoundWidget.h"
 
+class UGSTab;
+
 class SEmptyTab : public SCompoundWidget
 {
 public:
-
 	SLATE_BEGIN_ARGS(SEmptyTab) {}
+		SLATE_ARGUMENT(UGSTab*, Tab) 
 	SLATE_END_ARGS()
 
-	/**
-	 * Constructs the widget.
-	 */
 	void Construct(const FArguments& InArgs);
 
 private:
-
 	FReply OnOpenProjectClicked();
 
+	UGSTab* Tab;
 };
diff --git a/Engine/Source/Programs/SlateUGS/Private/Widgets/SGameSyncTab.cpp b/Engine/Source/Programs/SlateUGS/Private/Widgets/SGameSyncTab.cpp
--- a/Engine/Source/Programs/SlateUGS/Private/Widgets/SGameSyncTab.cpp
+++ b/Engine/Source/Programs/SlateUGS/Private/Widgets/SGameSyncTab.cpp
@@ -17,7 +17,7 @@
 
 TSharedRef<ITableRow> SGameSyncTab::GenerateHordeBuildTableRow(TSharedPtr<HordeBuildRowInfo> InItem, const TSharedRef<STableViewBase>& InOwnerTable)
 {
-	return SNew(STableRow<TSharedPtr<HordeBuildRowInfo>>, InOwnerTable)
+	return SNew(STableRow<TSharedPtr<HordeBuildRowInfo>>, InOwnerTable) // Todo: Maybe replace with SMultiColumnTableRow
 	[
 		SNew(SHorizontalBox)
 		// Build status
@@ -107,8 +107,7 @@ TSharedRef<ITableRow> SGameSyncTab::GenerateHordeBuildTableRow(TSharedPtr<HordeB
 
 void SGameSyncTab::Construct(const FArguments& InArgs)
 {
-	// Todo: Get horde builds from querying Perforce
-	SHordeBuilds = InArgs._HordeBuilds;
+	HordeBuilds = InArgs._HordeBuilds;
 
 	this->ChildSlot
 	[
@@ -404,8 +403,8 @@ void SGameSyncTab::Construct(const FArguments& InArgs)
 		+SVerticalBox::Slot()
 		.Padding(20.0f, 5.0f)
 		[
-			SAssignNew(SHordeBuildsView, SListView<TSharedPtr<HordeBuildRowInfo>>)
-			.ListItemsSource(&SHordeBuilds)
+			SAssignNew(HordeBuildsView, SListView<TSharedPtr<HordeBuildRowInfo>>)
+			.ListItemsSource(&HordeBuilds)
 			.OnGenerateRow(this, &SGameSyncTab::GenerateHordeBuildTableRow)
 		]
 		// Console window
diff --git a/Engine/Source/Programs/SlateUGS/Private/Widgets/SGameSyncTab.h b/Engine/Source/Programs/SlateUGS/Private/Widgets/SGameSyncTab.h
--- a/Engine/Source/Programs/SlateUGS/Private/Widgets/SGameSyncTab.h
+++ b/Engine/Source/Programs/SlateUGS/Private/Widgets/SGameSyncTab.h
@@ -8,22 +8,17 @@
 class SGameSyncTab : public SCompoundWidget
 {
 public:
-
 	SLATE_BEGIN_ARGS(SGameSyncTab) {}
-		SLATE_ARGUMENT(TArray<TSharedPtr<HordeBuildRowInfo>>, HordeBuilds)
+		SLATE_ARGUMENT(TArray<TSharedPtr<HordeBuildRowInfo>>, HordeBuilds) // Todo: remove once we start getting real horde data
 	SLATE_END_ARGS()
 
-	/**
-	 * Constructs the widget.
-	 */
 	void Construct(const FArguments& InArgs);
 
 private:
-
 	TSharedRef<ITableRow> GenerateHordeBuildTableRow(TSharedPtr<HordeBuildRowInfo> InItem, const TSharedRef<STableViewBase>& InOwnerTable);
 
-	TSharedPtr<SListView<TSharedPtr<HordeBuildRowInfo>>> SHordeBuildsView; 
-	TArray<TSharedPtr<HordeBuildRowInfo>> SHordeBuilds;
+	TSharedPtr<SListView<TSharedPtr<HordeBuildRowInfo>>> HordeBuildsView; 
+	TArray<TSharedPtr<HordeBuildRowInfo>> HordeBuilds;
 
 	static constexpr float HordeBuildRowHorizontalPadding = 10.0f;
 	static constexpr float HordeBuildRowVerticalPadding = 2.5f;
diff --git a/Engine/Source/Programs/SlateUGS/Private/Widgets/SWorkspaceWindow.cpp b/Engine/Source/Programs/SlateUGS/Private/Widgets/SWorkspaceWindow.cpp
--- a/Engine/Source/Programs/SlateUGS/Private/Widgets/SWorkspaceWindow.cpp
+++ b/Engine/Source/Programs/SlateUGS/Private/Widgets/SWorkspaceWindow.cpp
@@ -3,16 +3,18 @@
 #include "SWorkspaceWindow.h"
 
 #include "Styling/AppStyle.h"
-
 #include "DesktopPlatformModule.h"
 
-#include "../HordeBuildRowInfo.h"
+#include "UGSTab.h"
+#include "HordeBuildRowInfo.h"
 #include "SGameSyncTab.h"
 
 #define LOCTEXT_NAMESPACE "UGSWorkspaceWindow"
 
 void SWorkspaceWindow::Construct(const FArguments& InArgs)
 {
+	Tab = InArgs._Tab;
+
 	this->ChildSlot
 	[
 		SNew(SBox)
@@ -60,6 +62,11 @@ void SWorkspaceWindow::Construct(const FArguments& InArgs)
 					[
 						SAssignNew(LocalFileText, SEditableTextBox)
 						.HintText(LOCTEXT("FilePathHint", "Path/To/ProjectFile.uproject")) // Todo: Make hint text use backslash for Windows, forward slash for Unix
+			 			.OnTextChanged_Lambda([this](const FText& InText)
+						{
+							WorkspacePathText = InText.ToString();
+							fprintf(stderr, "Workspace path text changed to: %s\n", TCHAR_TO_ANSI(*WorkspacePathText));
+						})
 					]
 					+SHorizontalBox::Slot()
 					.FillWidth(2)
@@ -180,7 +187,7 @@ void SWorkspaceWindow::Construct(const FArguments& InArgs)
 					SNew(SButton)
 					.HAlign(HAlign_Center)
 					.Text(LOCTEXT("OkText", "Ok"))
-					.OnClicked(this, &SWorkspaceWindow::OnOkClicked)
+					.OnClicked(FOnClicked::CreateRaw(Tab, &UGSTab::OnWorkspaceChosen, WorkspacePathText))
 				]
 				+SHorizontalBox::Slot()
 				.HAlign(HAlign_Fill)
@@ -197,26 +204,6 @@ void SWorkspaceWindow::Construct(const FArguments& InArgs)
 	];
 }
 
-FReply SWorkspaceWindow::OnOkClicked()
-{
-	// Todo: Validate data
-	bool bIsDataValid = true;
-	if (bIsDataValid)
-	{
-		const TSharedRef<FGlobalTabmanager> TabManager = FGlobalTabmanager::Get();
-		
-		TabManager->SetActiveTab(TabManager->TryInvokeTab(FName("ActiveTab"))); // Todo: Pipe data into next window
-		TabManager->FindExistingLiveTab(FName("EmptyTab"))->RequestCloseTab();
-		FSlateApplication::Get().FindWidgetWindow(AsShared())->RequestDestroyWindow();
-	}
-	else
-	{
-		// Todo: Spawn error window telling user which data was invalid
-	}
-
-	return FReply::Handled();
-}
-
 FReply SWorkspaceWindow::OnCancelClicked()
 {
 	FSlateApplication::Get().FindWidgetWindow(AsShared())->RequestDestroyWindow();
diff --git a/Engine/Source/Programs/SlateUGS/Private/Widgets/SWorkspaceWindow.h b/Engine/Source/Programs/SlateUGS/Private/Widgets/SWorkspaceWindow.h
--- a/Engine/Source/Programs/SlateUGS/Private/Widgets/SWorkspaceWindow.h
+++ b/Engine/Source/Programs/SlateUGS/Private/Widgets/SWorkspaceWindow.h
@@ -4,26 +4,26 @@
 
 #include "Widgets/SCompoundWidget.h"
 
+class UGSTab;
+
 class SWorkspaceWindow : public SCompoundWidget
 {
 public:
-
 	SLATE_BEGIN_ARGS(SWorkspaceWindow) {}
+		SLATE_ARGUMENT(UGSTab*, Tab) 
 	SLATE_END_ARGS()
 
-	/**
-	 * Constructs the widget.
-	 */
 	void Construct(const FArguments& InArgs);
 
 private:
-
 	FReply OnBrowseClicked();
 
 	FReply OnOkClicked();
 	FReply OnCancelClicked();
 
 	bool bIsLocalFileSelected = true;
-
 	TSharedPtr<SEditableTextBox> LocalFileText = nullptr;
+	FString WorkspacePathText;
+
+	UGSTab* Tab;
 };
