Pass by ref vs value, this was causing default string value to always be used vs updated value

diff --git a/Engine/Source/Programs/SlateUGS/Private/UGSTab.cpp b/Engine/Source/Programs/SlateUGS/Private/UGSTab.cpp
--- a/Engine/Source/Programs/SlateUGS/Private/UGSTab.cpp
+++ b/Engine/Source/Programs/SlateUGS/Private/UGSTab.cpp
@@ -28,7 +28,7 @@ FSpawnTabArgs UGSTab::GetTabArgs() const
 	return TabArgs; 
 }
 
-FReply UGSTab::OnWorkspaceChosen(FString Path)
+FReply UGSTab::OnWorkspaceChosen(const FString& Path)
 {
 	fprintf(stderr, "Workspace path text received by OnWorkspaceChosen is: %s\n", TCHAR_TO_ANSI(*Path)); // Todo: Actually get the path variable bound
 	bool bIsDataValid = !Path.IsEmpty(); // Todo: Actually validate the data
diff --git a/Engine/Source/Programs/SlateUGS/Private/UGSTab.h b/Engine/Source/Programs/SlateUGS/Private/UGSTab.h
--- a/Engine/Source/Programs/SlateUGS/Private/UGSTab.h
+++ b/Engine/Source/Programs/SlateUGS/Private/UGSTab.h
@@ -2,7 +2,7 @@
 
 #pragma once
 
-#include "Widgets/SEmptyTab.h" 
+#include "Widgets/SEmptyTab.h"
 #include "Widgets/SGameSyncTab.h"
 #include "Widgets/SWorkspaceWindow.h"
 #include "UGSCore/GameSyncController.h"
@@ -18,7 +18,7 @@ public:
 	FSpawnTabArgs GetTabArgs() const;
 
 	// Slate callbacks
-	FReply OnWorkspaceChosen(FString Path);
+	FReply OnWorkspaceChosen(const FString& Path);
 private:
 
 	// Slate Data
@@ -30,5 +30,5 @@ private:
 	TSharedRef<SWorkspaceWindow> WorkspaceWindowView;
 
 	// Controller data
-	GameSyncController SyncController; 
+	GameSyncController SyncController;
 };
diff --git a/Engine/Source/Programs/SlateUGS/Private/Widgets/SWorkspaceWindow.cpp b/Engine/Source/Programs/SlateUGS/Private/Widgets/SWorkspaceWindow.cpp
--- a/Engine/Source/Programs/SlateUGS/Private/Widgets/SWorkspaceWindow.cpp
+++ b/Engine/Source/Programs/SlateUGS/Private/Widgets/SWorkspaceWindow.cpp
@@ -187,7 +187,7 @@ void SWorkspaceWindow::Construct(const FArguments& InArgs)
 					SNew(SButton)
 					.HAlign(HAlign_Center)
 					.Text(LOCTEXT("OkText", "Ok"))
-					.OnClicked(FOnClicked::CreateRaw(Tab, &UGSTab::OnWorkspaceChosen, WorkspacePathText))
+					.OnClicked(FOnClicked::CreateRaw<UGSTab, const FString&>(Tab, &UGSTab::OnWorkspaceChosen, WorkspacePathText))
 				]
 				+SHorizontalBox::Slot()
 				.HAlign(HAlign_Fill)
