Cleaned up callback for clicking Ok after selecting a workspace, separated model code (workspace path validation) from view code (slate widgets)

diff --git a/Engine/Source/Programs/SlateUGS/Private/UGSCore/GameSyncController.cpp b/Engine/Source/Programs/SlateUGS/Private/UGSCore/GameSyncController.cpp
--- a/Engine/Source/Programs/SlateUGS/Private/UGSCore/GameSyncController.cpp
+++ b/Engine/Source/Programs/SlateUGS/Private/UGSCore/GameSyncController.cpp
@@ -1,3 +1,5 @@
+// Copyright Epic Games, Inc. All Rights Reserved.
+
 #include "GameSyncController.h"
 
 #include "BuildStep.h"
diff --git a/Engine/Source/Programs/SlateUGS/Private/UGSTab.cpp b/Engine/Source/Programs/SlateUGS/Private/UGSTab.cpp
--- a/Engine/Source/Programs/SlateUGS/Private/UGSTab.cpp
+++ b/Engine/Source/Programs/SlateUGS/Private/UGSTab.cpp
@@ -7,8 +7,7 @@
 UGSTab::UGSTab() : TabArgs(nullptr, FTabId()),
 				   TabWidget(SNew(SDockTab)),
 				   EmptyTabView(SNew(SEmptyTab).Tab(this)),
-				   GameSyncTabView(SNew(SGameSyncTab)),
-				   WorkspaceWindowView(SNew(SWorkspaceWindow).Tab(this))
+				   GameSyncTabView(SNew(SGameSyncTab))
 {
 	TabWidget->SetContent(EmptyTabView);
 }
@@ -28,53 +27,16 @@ FSpawnTabArgs UGSTab::GetTabArgs() const
 	return TabArgs; 
 }
 
-FReply UGSTab::OnWorkspaceChosen(const FString& Path)
+bool UGSTab::OnWorkspaceChosen(const FString& Path)
 {
-	fprintf(stderr, "Workspace path text received by OnWorkspaceChosen is: %s\n", TCHAR_TO_ANSI(*Path)); // Todo: Actually get the path variable bound
 	bool bIsDataValid = !Path.IsEmpty(); // Todo: Actually validate the data
 	if (bIsDataValid)
 	{
 		TabWidget->SetContent(GameSyncTabView); // Todo: Set GameSyncTabView data
+		return true;
 	}
-	else
-	{
-		// Todo: factor out into an error window widget class (check if one already exists)
-		TSharedRef<SWindow> Window = SNew(SWindow)
-		.Title(LOCTEXT("ErrorWindowTitle", "Error Opening Project"))
-		.SizingRule(ESizingRule::Autosized)
-		.MaxWidth(400.0f)
-		[
-			SNew(SVerticalBox)
-			+SVerticalBox::Slot()
-			.AutoHeight()
-			.HAlign(HAlign_Center)
-			.Padding(10.0f, 10.0f)
-			[
-				SNew(STextBlock)
-				.AutoWrapText(true)
-				.Text(LOCTEXT("ErrorText", "Error opening .uproject file, try again")) // Todo: detect the actual reason for error and report it here
-			]
-			+SVerticalBox::Slot()
-			.AutoHeight()
-			.HAlign(HAlign_Center)
-			.Padding(0.0f, 10.0f, 0.0f, 0.0f)
-			[
-				SNew(SButton)
-				.HAlign(HAlign_Center)
-				.Text(LOCTEXT("ErrorWindowOkayButtonText", "Ok"))
-				.OnClicked_Lambda([&Window]() 
-				{
-					Window->RequestDestroyWindow();
-					return FReply::Handled(); 
-				})
-			]
-		];
-
-		FSlateApplication::Get().AddModalWindow(Window, FSlateApplication::Get().GetActiveModalWindow(), false); // Todo: Figure out parent window for modal window
-		// Todo: loading screen widget for detect settings
-	}
-
-	return FReply::Handled();
+	
+	return false;
 }
 
 #undef LOCTEXT_NAMESPACE
diff --git a/Engine/Source/Programs/SlateUGS/Private/UGSTab.h b/Engine/Source/Programs/SlateUGS/Private/UGSTab.h
--- a/Engine/Source/Programs/SlateUGS/Private/UGSTab.h
+++ b/Engine/Source/Programs/SlateUGS/Private/UGSTab.h
@@ -18,7 +18,7 @@ public:
 	FSpawnTabArgs GetTabArgs() const;
 
 	// Slate callbacks
-	FReply OnWorkspaceChosen(const FString& Path);
+	bool OnWorkspaceChosen(const FString& Path);
 private:
 
 	// Slate Data
@@ -27,7 +27,6 @@ private:
 
 	TSharedRef<SEmptyTab> EmptyTabView;
 	TSharedRef<SGameSyncTab> GameSyncTabView;
-	TSharedRef<SWorkspaceWindow> WorkspaceWindowView;
 
 	// Controller data
 	GameSyncController SyncController;
diff --git a/Engine/Source/Programs/SlateUGS/Private/Widgets/SEmptyTab.cpp b/Engine/Source/Programs/SlateUGS/Private/Widgets/SEmptyTab.cpp
--- a/Engine/Source/Programs/SlateUGS/Private/Widgets/SEmptyTab.cpp
+++ b/Engine/Source/Programs/SlateUGS/Private/Widgets/SEmptyTab.cpp
@@ -71,7 +71,7 @@ FReply SEmptyTab::OnOpenProjectClicked()
 	.Title(LOCTEXT("WindowTitle", "Open Project"))
 	.SizingRule(ESizingRule::Autosized)
 	[
-		SNew(SWorkspaceWindow)
+		SNew(SWorkspaceWindow).Tab(Tab)
 	];
 
 	FSlateApplication& SlateApplication = FSlateApplication::Get();
diff --git a/Engine/Source/Programs/SlateUGS/Private/Widgets/SWorkspaceWindow.cpp b/Engine/Source/Programs/SlateUGS/Private/Widgets/SWorkspaceWindow.cpp
--- a/Engine/Source/Programs/SlateUGS/Private/Widgets/SWorkspaceWindow.cpp
+++ b/Engine/Source/Programs/SlateUGS/Private/Widgets/SWorkspaceWindow.cpp
@@ -187,7 +187,7 @@ void SWorkspaceWindow::Construct(const FArguments& InArgs)
 					SNew(SButton)
 					.HAlign(HAlign_Center)
 					.Text(LOCTEXT("OkText", "Ok"))
-					.OnClicked(FOnClicked::CreateRaw<UGSTab, const FString&>(Tab, &UGSTab::OnWorkspaceChosen, WorkspacePathText))
+					.OnClicked(this, &SWorkspaceWindow::OnOkClicked)
 				]
 				+SHorizontalBox::Slot()
 				.HAlign(HAlign_Fill)
@@ -204,6 +204,54 @@ void SWorkspaceWindow::Construct(const FArguments& InArgs)
 	];
 }
 
+FReply SWorkspaceWindow::OnOkClicked()
+{
+	bool bIsWorkspaceValid = Tab->OnWorkspaceChosen(WorkspacePathText);
+	if (bIsWorkspaceValid)
+	{
+		FSlateApplication::Get().FindWidgetWindow(AsShared())->RequestDestroyWindow();
+	}
+	else
+	{
+		// Todo: factor out into an error window widget class (check if one already exists)
+		TSharedRef<SWindow> Window = SNew(SWindow)
+		.Title(LOCTEXT("ErrorWindowTitle", "Error Opening Project"))
+		.SizingRule(ESizingRule::Autosized)
+		.MaxWidth(400.0f)
+		[
+			SNew(SVerticalBox)
+			+SVerticalBox::Slot()
+			.AutoHeight()
+			.HAlign(HAlign_Center)
+			.Padding(10.0f, 10.0f)
+			[
+				SNew(STextBlock)
+				.AutoWrapText(true)
+				.Text(LOCTEXT("ErrorText", "Error opening .uproject file, try again")) // Todo: detect the actual reason for error and report it here
+			]
+			+SVerticalBox::Slot()
+			.AutoHeight()
+			.HAlign(HAlign_Center)
+			.Padding(0.0f, 10.0f, 0.0f, 0.0f)
+			[
+				SNew(SButton)
+				.HAlign(HAlign_Center)
+				.Text(LOCTEXT("ErrorWindowOkayButtonText", "Ok"))
+				.OnClicked_Lambda([&Window]() 
+				{
+					Window->RequestDestroyWindow();
+					return FReply::Handled(); 
+				})
+			]
+		];
+
+		FSlateApplication::Get().AddModalWindow(Window, FSlateApplication::Get().GetActiveModalWindow(), false); // Todo: Figure out parent window for modal window
+		// Todo: loading screen widget for detect settings
+	}
+
+	return FReply::Handled();
+}
+
 FReply SWorkspaceWindow::OnCancelClicked()
 {
 	FSlateApplication::Get().FindWidgetWindow(AsShared())->RequestDestroyWindow();
