Fix crash on exit due to TabManager going out of scope after the Engine is dead

diff --git a/Engine/Source/Programs/SlateUGS/Private/SlateUGSApp.cpp b/Engine/Source/Programs/SlateUGS/Private/SlateUGSApp.cpp
--- a/Engine/Source/Programs/SlateUGS/Private/SlateUGSApp.cpp
+++ b/Engine/Source/Programs/SlateUGS/Private/SlateUGSApp.cpp
@@ -17,14 +17,14 @@ IMPLEMENT_APPLICATION(SlateUGS, "SlateUGS");
 
 int RunSlateUGS(const TCHAR* CommandLine)
 {
-	FTaskTagScope TaskTagScope(ETaskTag::EGameThread);	
+	FTaskTagScope TaskTagScope(ETaskTag::EGameThread);
 
 	// start up the main loop
 	GEngineLoop.PreInit(CommandLine);
 
 	// Make sure all UObject classes are registered and default properties have been initialized
 	ProcessNewlyLoadedUObjects();
-	
+
 	// Tell the module manager it may now process newly-loaded UObjects when new C++ modules are loaded
 	FModuleManager::Get().StartProcessingNewlyLoadedObjects();
 
@@ -38,31 +38,34 @@ int RunSlateUGS(const TCHAR* CommandLine)
 
 	FAppStyle::SetAppStyleSetName(FAppStyle::GetAppStyleSetName());
 
-	// Build the slate UI for the program window
-	UGSTabManager TabManager; 
-	TabManager.ConstructTabs();
-
-	// loop while the server does the rest
-	while (!IsEngineExitRequested())
+	// new scope to allow TabManager to go out of scope before the Engine is dead
 	{
-		BeginExitIfRequested();
-
-		FTaskGraphInterface::Get().ProcessThreadUntilIdle(ENamedThreads::GameThread);
-		FStats::AdvanceFrame(false);
-		FTSTicker::GetCoreTicker().Tick(FApp::GetDeltaTime());
-		FSlateApplication::Get().PumpMessages();
-		FSlateApplication::Get().Tick();		
-		FPlatformProcess::Sleep(0.01f);
-	
-		GFrameCounter++;
+		// Build the slate UI for the program window
+		UGSTabManager TabManager;
+		TabManager.ConstructTabs();
+
+		// loop while the server does the rest
+		while (!IsEngineExitRequested())
+		{
+			BeginExitIfRequested();
+
+			FTaskGraphInterface::Get().ProcessThreadUntilIdle(ENamedThreads::GameThread);
+			FStats::AdvanceFrame(false);
+			FTSTicker::GetCoreTicker().Tick(FApp::GetDeltaTime());
+			FSlateApplication::Get().PumpMessages();
+			FSlateApplication::Get().Tick();
+			FPlatformProcess::Sleep(0.01f);
+
+			GFrameCounter++;
+		}
 	}
 
 	FCoreDelegates::OnExit.Broadcast();
 	FSlateApplication::Shutdown();
 	FModuleManager::Get().UnloadModulesAtShutdown();
 
-	GEngineLoop.AppPreExit();
-	GEngineLoop.AppExit();
+	FEngineLoop::AppPreExit();
+	FEngineLoop::AppExit();
 
 	return 0;
 }
