Opening a valid project runs DetectSettings now. GameSyncController code will need to be deleted soon

diff --git a/Engine/Source/Programs/SlateUGS/Private/UGSTab.cpp b/Engine/Source/Programs/SlateUGS/Private/UGSTab.cpp
--- a/Engine/Source/Programs/SlateUGS/Private/UGSTab.cpp
+++ b/Engine/Source/Programs/SlateUGS/Private/UGSTab.cpp
@@ -2,6 +2,11 @@
 
 #include "UGSTab.h"
 
+#include "UGSCore/BuildStep.h"
+#include "UGSCore/DetectProjectSettingsTask.h"
+#include "UGSCore/Utility.h"
+#include "UGSCore/UserSettings.h"
+
 #define LOCTEXT_NAMESPACE "UGSTab"
 
 UGSTab::UGSTab() : TabArgs(nullptr, FTabId()),
@@ -27,11 +32,13 @@ FSpawnTabArgs UGSTab::GetTabArgs() const
 	return TabArgs; 
 }
 
-bool UGSTab::OnWorkspaceChosen(const FString& Path)
+bool UGSTab::OnWorkspaceChosen(const FString& Project)
 {
-	bool bIsDataValid = !Path.IsEmpty(); // Todo: Actually validate the data
+	bool bIsDataValid = FPaths::FileExists(Project); // Todo: Check that the project file is also associated with a workspace
 	if (bIsDataValid)
 	{
+		ProjectFileName = Project;
+		SetupWorkspace();
 		TabWidget->SetContent(GameSyncTabView); // Todo: Set GameSyncTabView data
 		return true;
 	}
@@ -39,4 +46,175 @@ bool UGSTab::OnWorkspaceChosen(const FString& Path)
 	return false;
 }
 
+namespace
+{
+	// TODO super hacky... fix up later
+	#if PLATFORM_WINDOWS
+		const TCHAR* HostPlatform = TEXT("Win64");
+	#elif PLATFORM_MAC
+		const TCHAR* HostPlatform = TEXT("Mac");
+	#else
+		const TCHAR* HostPlatform = TEXT("Linux");
+	#endif
+
+	class FLineWriter : public FLineBasedTextWriter
+	{
+	public:
+		virtual void FlushLine(const FString& Line) override
+		{
+			printf("%s\n", TCHAR_TO_ANSI(*Line));
+		}
+	};
+
+	// TODO move most of these to a class with access to the settings
+	//bool ShouldSyncPrecompiledEditor(TSharedPtr<FUserSettings> Settings)
+	bool ShouldSyncPrecompiledEditor()
+	{
+		return false;
+		//return Settings->bSyncPrecompiledEditor;// && PerforceMonitor->HasZippedBinaries();
+	}
+
+	// Honestly ... seems ... super hacky/hardcoded. With out all these you assert when trying to merge build targets sooo a bit odd
+	// TODO Need to do each of these ... per ... platform??
+	TMap<FGuid, FCustomConfigObject> GetDefaultBuildStepObjects(const FString& EditorTargetName, TSharedPtr<FUserSettings> Settings)
+	{
+		TArray<FBuildStep> DefaultBuildSteps;
+		DefaultBuildSteps.Add(FBuildStep(FGuid(0x01F66060, 0x73FA4CC8, 0x9CB3E217, 0xFBBA954E), 0, TEXT("Compile UnrealHeaderTool"), TEXT("Compiling UnrealHeaderTool..."), 1, TEXT("UnrealHeaderTool"), HostPlatform, TEXT("Development"), TEXT(""), !ShouldSyncPrecompiledEditor()));
+		FString ActualEditorTargetName = (EditorTargetName.Len() > 0)? EditorTargetName : "UnrealEditor";
+		DefaultBuildSteps.Add(FBuildStep(FGuid(0xF097FF61, 0xC9164058, 0x839135B4, 0x6C3173D5), 1, FString::Printf(TEXT("Compile %s"), *ActualEditorTargetName), FString::Printf(TEXT("Compiling %s..."), *ActualEditorTargetName), 10, ActualEditorTargetName, HostPlatform, ::ToString(Settings->CompiledEditorBuildConfig), TEXT(""), !ShouldSyncPrecompiledEditor()));
+		DefaultBuildSteps.Add(FBuildStep(FGuid(0xC6E633A1, 0x956F4AD3, 0xBC956D06, 0xD131E7B4), 2, TEXT("Compile ShaderCompileWorker"), TEXT("Compiling ShaderCompileWorker..."), 1, TEXT("ShaderCompileWorker"), HostPlatform, TEXT("Development"), TEXT(""), !ShouldSyncPrecompiledEditor()));
+		DefaultBuildSteps.Add(FBuildStep(FGuid(0x24FFD88C, 0x79014899, 0x9696AE10, 0x66B4B6E8), 3, TEXT("Compile UnrealLightmass"), TEXT("Compiling UnrealLightmass..."), 1, TEXT("UnrealLightmass"), HostPlatform, TEXT("Development"), TEXT(""), !ShouldSyncPrecompiledEditor()));
+		DefaultBuildSteps.Add(FBuildStep(FGuid(0xFFF20379, 0x06BF4205, 0x8A3EC534, 0x27736688), 4, TEXT("Compile CrashReportClient"), TEXT("Compiling CrashReportClient..."), 1, TEXT("CrashReportClient"), HostPlatform, TEXT("Development"), TEXT(""), !ShouldSyncPrecompiledEditor()));
+		DefaultBuildSteps.Add(FBuildStep(FGuid(0x89FE8A79, 0xD2594C7B, 0xBFB468F7, 0x218B91C2), 5, TEXT("Compile UnrealInsights"), TEXT("Compiling UnrealInsights..."), 1, TEXT("UnrealInsights"), HostPlatform, TEXT("Development"), TEXT(""), !ShouldSyncPrecompiledEditor()));
+		DefaultBuildSteps.Add(FBuildStep(FGuid(0x46312669, 0x5069428D, 0x8D72C241, 0x6C5A322E), 6, TEXT("Launch UnrealInsights"), TEXT("Running UnrealInsights..."), 1, TEXT("UnrealInsights"), HostPlatform, TEXT("Shipping"), TEXT(""), !ShouldSyncPrecompiledEditor()));
+		DefaultBuildSteps.Add(FBuildStep(FGuid(0xBB48CA5B, 0x56824432, 0x824DC451, 0x336A6523), 7, TEXT("Compile Zen Dashboard"), TEXT("Compile ZenDashboard Step..."), 1, TEXT("ZenDashboard"), HostPlatform, TEXT("Development"), TEXT(""), !ShouldSyncPrecompiledEditor()));
+		DefaultBuildSteps.Add(FBuildStep(FGuid(0x586CC0D3, 0x39144DF9, 0xACB62C02, 0xCD9D4FC6), 8, TEXT("Launch Zen Dashboard"), TEXT("Running Zen Dashboard..."), 1, TEXT("ZenDashboard"), HostPlatform, TEXT("Development"), TEXT(""), !ShouldSyncPrecompiledEditor()));
+		DefaultBuildSteps.Add(FBuildStep(FGuid(0x91C2A429, 0xC39149B4, 0x92A54E6B, 0xE71E0F00), 9, TEXT("Compile SwitchboardListener"), TEXT("Compiling SwitchboardListener..."), 1, TEXT("SwitchboardListener"), HostPlatform, TEXT("Development"), TEXT(""), !ShouldSyncPrecompiledEditor()));
+		DefaultBuildSteps.Add(FBuildStep(FGuid(0x5036C75B, 0x8DF04329, 0x82A1869D, 0xD2D48605), 10, TEXT("Compile UnrealMultiUserServer"), TEXT("Compiling UnrealMultiUserServer..."), 1, TEXT("UnrealMultiUserServer"), HostPlatform, TEXT("Development"), TEXT(""), !ShouldSyncPrecompiledEditor()));
+		DefaultBuildSteps.Add(FBuildStep(FGuid(0x274B89C3, 0x9DC64465, 0xA50840AB, 0xC4593CC2), 11, TEXT("Compile UnrealMultiUserSlateServer"), TEXT("Compiling UnrealMultiUserSlateServer..."), 1, TEXT("UnrealMultiUserSlateServer"), HostPlatform, TEXT("Development"), TEXT(""), !ShouldSyncPrecompiledEditor()));
+
+		TMap<FGuid,FCustomConfigObject> DefaultBuildStepObjects;
+		for(const FBuildStep& DefaultBuildStep : DefaultBuildSteps)
+		{
+			DefaultBuildStepObjects.Add(DefaultBuildStep.UniqueId, DefaultBuildStep.ToConfigObject());
+		}
+		return DefaultBuildStepObjects;
+	}
+
+	EBuildConfig GetEditorBuildConfig()
+	{
+		return EBuildConfig::Development;
+		//return ShouldSyncPrecompiledEditor()? EBuildConfig::Development : Settings->CompiledEditorBuildConfig;
+	}
+
+	FString GetEditorExePath(EBuildConfig Config, TSharedPtr<FDetectProjectSettingsTask> DetectSettings)
+	{
+		FString ExeFileName = TEXT("UnrealEditor");
+
+		if(Config != EBuildConfig::DebugGame && Config != EBuildConfig::Development)
+		{
+	#if PLATFORM_WINDOWS
+			ExeFileName = FString::Printf(TEXT("UnrealEditor-%s-%s.exe"), HostPlatform, *::ToString(Config));
+	#else
+			ExeFileName = FString::Printf(TEXT("UnrealEditor-%s-%s"), HostPlatform, *::ToString(Config));
+	#endif
+		}
+		return DetectSettings->BranchDirectoryName / "Engine" / "Binaries" / HostPlatform / ExeFileName;
+	}
+
+	TMap<FString, FString> GetWorkspaceVariables(TSharedPtr<FDetectProjectSettingsTask> DetectSettings)
+	{
+		EBuildConfig EditorBuildConfig = GetEditorBuildConfig();
+
+		TMap<FString, FString> Variables;
+		Variables.Add("BranchDir", DetectSettings->BranchDirectoryName);
+		Variables.Add("ProjectDir", FPaths::GetPath(DetectSettings->NewSelectedFileName));
+		Variables.Add("ProjectFile", DetectSettings->NewSelectedFileName);
+
+		// Todo: These might not be called "UE4*" anymore
+		Variables.Add("UE4EditorExe", GetEditorExePath(EditorBuildConfig, DetectSettings));
+		Variables.Add("UE4EditorCmdExe", GetEditorExePath(EditorBuildConfig, DetectSettings).Replace(TEXT(".exe"), TEXT("-Cmd.exe")));
+		Variables.Add("UE4EditorConfig", ::ToString(EditorBuildConfig));
+		Variables.Add("UE4EditorDebugArg", (EditorBuildConfig == EBuildConfig::Debug || EditorBuildConfig == EBuildConfig::DebugGame)? " -debug" : "");
+		return Variables;
+	}
+}
+
+void UGSTab::SetupWorkspace()
+{
+	ProjectFileName = FUtility::GetPathWithCorrectCase(ProjectFileName);
+	TSharedPtr<FDetectProjectSettingsTask> DetectSettings = MakeShared<FDetectProjectSettingsTask>(MakeShared<FPerforceConnection>(TEXT(""), TEXT(""), TEXT("")), ProjectFileName, MakeShared<FLineWriter>());
+
+	// AbortEvent for canceling but honestly not setup well.
+	// Todo: fix this up and likely enable exceptions
+	FEvent* AbortEvent = nullptr;
+	DetectSettings->Run(AbortEvent);
+
+	FString DataFolder = FString(FPlatformProcess::UserSettingsDir()) / TEXT("UnrealGameSync");
+	IFileManager::Get().MakeDirectory(*DataFolder);
+
+	TSharedPtr<FUserSettings> Settings = MakeShared<FUserSettings>(*(DataFolder / TEXT("UnrealGameSync.ini")));
+
+	TSharedRef<FPerforceConnection> PerforceClient = DetectSettings->PerforceClient.ToSharedRef();
+	TSharedRef<FUserWorkspaceSettings> WorkspaceSettings = Settings->FindOrAddWorkspace(*DetectSettings->BranchClientPath);
+	TSharedRef<FUserProjectSettings> ProjectSettings = Settings->FindOrAddProject(*DetectSettings->NewSelectedClientFileName);
+
+	FString SelectedClientFileName = DetectSettings->NewSelectedClientFileName;
+	FString BranchDirectoryName = DetectSettings->BranchDirectoryName;
+	FString BranchClientPath = DetectSettings->BranchClientPath;
+	FString SelectedProjectIdentifier = DetectSettings->NewSelectedProjectIdentifier;
+	FString EditorTarget = DetectSettings->NewProjectEditorTarget;
+
+	// Check if the project we've got open in this workspace is the one we're actually synced to
+	int CurrentChangeNumber = -1;
+	if(WorkspaceSettings->CurrentProjectIdentifier == SelectedProjectIdentifier)
+	{
+		CurrentChangeNumber = WorkspaceSettings->CurrentChangeNumber;
+	}
+
+	FString ClientKey = BranchClientPath.Replace(*FString::Printf(TEXT("//%s/"), *PerforceClient->ClientName), TEXT(""));
+	if(ClientKey.EndsWith(TEXT("/")))
+	{
+		ClientKey = ClientKey.Left(ClientKey.Len() - 1);
+	}
+
+	FString ProjectLogBaseName = DataFolder / FString::Printf(TEXT("%s@%s"), *PerforceClient->ClientName, *ClientKey.Replace(TEXT("/"), TEXT("$")));
+	FString TelemetryProjectIdentifier = FPerforceUtils::GetClientOrDepotDirectoryName(*SelectedProjectIdentifier);
+
+	Workspace = MakeShared<FWorkspace>(PerforceClient, BranchDirectoryName, ProjectFileName, BranchClientPath, SelectedClientFileName, CurrentChangeNumber, WorkspaceSettings->LastBuiltChangeNumber, TelemetryProjectIdentifier, MakeShared<FLineWriter>());
+
+	// Todo:Move into a sync operation
+	// TArray<FString> CombinedSyncFilter = FUserSettings::GetCombinedSyncFilter(Workspace->GetSyncCategories(), Settings->SyncView, Settings->SyncExcludedCategories, WorkspaceSettings->SyncView, WorkspaceSettings->SyncExcludedCategories);
+
+	// // Options on what to do with workspace when updating it
+	// EWorkspaceUpdateOptions Options = EWorkspaceUpdateOptions::Sync | EWorkspaceUpdateOptions::SyncArchives | EWorkspaceUpdateOptions::GenerateProjectFiles;
+	// if(Settings->bAutoResolveConflicts)
+	// {
+	// 	Options |= EWorkspaceUpdateOptions::AutoResolveChanges;
+	// }
+	// if(Settings->bUseIncrementalBuilds)
+	// {
+	// 	Options |= EWorkspaceUpdateOptions::UseIncrementalBuilds;
+	// }
+	// if(Settings->bBuildAfterSync)
+	// {
+	// 	Options |= EWorkspaceUpdateOptions::Build;
+	// }
+	// if(Settings->bBuildAfterSync && Settings->bRunAfterSync)
+	// {
+	// 	Options |= EWorkspaceUpdateOptions::RunAfterSync; 
+	// }
+	// if(Settings->bOpenSolutionAfterSync)
+	// {
+	// 	Options |= EWorkspaceUpdateOptions::OpenSolutionAfterSync;
+	// }
+
+	// // Hacking in the CL
+	// int ChangeNumber = 20658981;
+	// TSharedRef<FWorkspaceUpdateContext, ESPMode::ThreadSafe> Context = MakeShared<FWorkspaceUpdateContext, ESPMode::ThreadSafe>(ChangeNumber, Options, CombinedSyncFilter, GetDefaultBuildStepObjects(EditorTarget, Settings), ProjectSettings->BuildSteps, TSet<FGuid>(), GetWorkspaceVariables(DetectSettings));
+
+	// // Update the workspace with the Context!
+	// Workspace->Update(Context);
+}
+
 #undef LOCTEXT_NAMESPACE
diff --git a/Engine/Source/Programs/SlateUGS/Private/UGSTab.h b/Engine/Source/Programs/SlateUGS/Private/UGSTab.h
--- a/Engine/Source/Programs/SlateUGS/Private/UGSTab.h
+++ b/Engine/Source/Programs/SlateUGS/Private/UGSTab.h
@@ -5,7 +5,8 @@
 #include "Widgets/SEmptyTab.h"
 #include "Widgets/SGameSyncTab.h"
 #include "Widgets/SWorkspaceWindow.h"
-#include "UGSCore/GameSyncController.h"
+
+#include "UGSCore/Workspace.h"
 
 class UGSTab
 {
@@ -18,16 +19,19 @@ public:
 	FSpawnTabArgs GetTabArgs() const;
 
 	// Slate callbacks
-	bool OnWorkspaceChosen(const FString& Path);
+	bool OnWorkspaceChosen(const FString& Project);
 private:
 
+	// Core functions
+	void SetupWorkspace();
+
 	// Slate Data
 	FSpawnTabArgs TabArgs;
 	TSharedRef<SDockTab> TabWidget;
-
 	TSharedRef<SEmptyTab> EmptyTabView;
 	TSharedRef<SGameSyncTab> GameSyncTabView;
 
-	// Controller data
-	GameSyncController SyncController;
+	// Core data
+	TSharedPtr<FWorkspace> Workspace;
+	FString ProjectFileName;
 };
