Moved GameSyncController contents into UGSTab and deleted GameSyncController

diff --git a/Engine/Source/Programs/SlateUGS/Private/UGSCore/GameSyncController.cpp b/Engine/Source/Programs/SlateUGS/Private/UGSCore/GameSyncController.cpp
deleted file mode 100644
diff --git a/Engine/Source/Programs/SlateUGS/Private/UGSCore/GameSyncController.h b/Engine/Source/Programs/SlateUGS/Private/UGSCore/GameSyncController.h
deleted file mode 100644
diff --git a/Engine/Source/Programs/SlateUGS/Private/UGSTab.cpp b/Engine/Source/Programs/SlateUGS/Private/UGSTab.cpp
--- a/Engine/Source/Programs/SlateUGS/Private/UGSTab.cpp
+++ b/Engine/Source/Programs/SlateUGS/Private/UGSTab.cpp
@@ -48,6 +48,72 @@ bool UGSTab::OnWorkspaceChosen(const FString& Project)
 
 namespace
 {
+	// Todo: may need this, was in GameSyncController.h
+	/*
+	#include "Perforce.h"
+	#include "CustomConfigFile.h"
+	#include "Misc/EnumClassFlags.h"
+
+	enum class EWorkspaceUpdateOptions
+	{
+		Sync = 0x01,
+		SyncSingleChange = 0x02,
+		AutoResolveChanges = 0x04,
+		GenerateProjectFiles = 0x08,
+		SyncArchives = 0x10,
+		Build = 0x20,
+		UseIncrementalBuilds = 0x40,
+		ScheduledBuild = 0x80,
+		RunAfterSync = 0x100,
+		OpenSolutionAfterSync = 0x200,
+		ContentOnly = 0x400
+	};
+
+	ENUM_CLASS_FLAGS(EWorkspaceUpdateOptions)
+
+	enum class EWorkspaceUpdateResult
+	{
+		Canceled,
+		FailedToSync,
+		FilesToResolve,
+		FilesToClobber,
+		FailedToCompile,
+		FailedToCompileWithCleanWorkspace,
+		Success,
+	};
+
+	FString ToString(EWorkspaceUpdateResult WorkspaceUpdateResult);
+	bool TryParse(const TCHAR* Text, EWorkspaceUpdateResult& OutWorkspaceUpdateResult);
+
+	struct FWorkspaceUpdateContext
+	{
+		FDateTime StartTime;
+		int ChangeNumber;
+		EWorkspaceUpdateOptions Options;
+		TArray<FString> SyncFilter;
+		TMap<FString, FString> ArchiveTypeToDepotPath;
+		TMap<FString, bool> ClobberFiles;
+		TMap<FGuid,FCustomConfigObject> DefaultBuildSteps;
+		TArray<FCustomConfigObject> UserBuildStepObjects;
+		TSet<FGuid> CustomBuildSteps;
+		TMap<FString, FString> Variables;
+		FPerforceSyncOptions PerforceSyncOptions;
+
+		FWorkspaceUpdateContext(int InChangeNumber, EWorkspaceUpdateOptions InOptions, const TArray<FString>& InSyncFilter, const TMap<FGuid, FCustomConfigObject>& InDefaultBuildSteps, const TArray<FCustomConfigObject>& InUserBuildSteps, const TSet<FGuid>& InCustomBuildSteps, const TMap<FString, FString>& InVariables);
+	};
+
+	struct FWorkspaceSyncCategory
+	{
+		FGuid UniqueId;
+		bool bEnable;
+		FString Name;
+		TArray<FString> Paths;
+
+		FWorkspaceSyncCategory(const FGuid& InUniqueId);
+		FWorkspaceSyncCategory(const FGuid& InUniqueId, const TCHAR* InName, const TCHAR* InPaths);
+	};
+	*/
+
 	// TODO super hacky... fix up later
 	#if PLATFORM_WINDOWS
 		const TCHAR* HostPlatform = TEXT("Win64");
