Add a Modal Task Window to pop something up if we are blocking

diff --git a/Engine/Source/Programs/SlateUGS/Private/UGSTab.cpp b/Engine/Source/Programs/SlateUGS/Private/UGSTab.cpp
--- a/Engine/Source/Programs/SlateUGS/Private/UGSTab.cpp
+++ b/Engine/Source/Programs/SlateUGS/Private/UGSTab.cpp
@@ -6,6 +6,7 @@
 #include "UGSCore/DetectProjectSettingsTask.h"
 #include "UGSCore/Utility.h"
 #include "UGSCore/UserSettings.h"
+#include "Widgets/SModalTaskWindow.h"
 
 #define LOCTEXT_NAMESPACE "UGSTab"
 
@@ -209,12 +210,14 @@ namespace
 void UGSTab::SetupWorkspace()
 {
 	ProjectFileName = FUtility::GetPathWithCorrectCase(ProjectFileName);
-	TSharedPtr<FDetectProjectSettingsTask> DetectSettings = MakeShared<FDetectProjectSettingsTask>(MakeShared<FPerforceConnection>(TEXT(""), TEXT(""), TEXT("")), ProjectFileName, MakeShared<FLineWriter>());
+	TSharedRef<FDetectProjectSettingsTask> DetectSettings = MakeShared<FDetectProjectSettingsTask>(MakeShared<FPerforceConnection>(TEXT(""), TEXT(""), TEXT("")), ProjectFileName, MakeShared<FLineWriter>());
 
-	// AbortEvent for canceling but honestly not setup well.
-	// Todo: fix this up and likely enable exceptions
-	FEvent* AbortEvent = nullptr;
-	DetectSettings->Run(AbortEvent);
+	TSharedRef<FModalTaskResult> Result = ExecuteModalTask(TabWidget, DetectSettings, LOCTEXT("OpeningProjectTitle", "Opening Project"), LOCTEXT("OpeningProjectCaption", "Opening project, please wait..."));
+	if(Result->Failed())
+	{
+		FMessageDialog::Open(EAppMsgType::Ok, Result->GetMessage());
+		return;
+	}
 
 	FString DataFolder = FString(FPlatformProcess::UserSettingsDir()) / TEXT("UnrealGameSync");
 	IFileManager::Get().MakeDirectory(*DataFolder);
diff --git a/Engine/Source/Programs/SlateUGS/Private/Widgets/SModalTaskWindow.cpp b/Engine/Source/Programs/SlateUGS/Private/Widgets/SModalTaskWindow.cpp
new file mode 100644
--- /dev/null
+++ b/Engine/Source/Programs/SlateUGS/Private/Widgets/SModalTaskWindow.cpp
@@ -0,0 +1,83 @@
+// Copyright Epic Games, Inc. All Rights Reserved.
+
+#include "SModalTaskWindow.h"
+#include "Framework/Application/SlateApplication.h"
+#include "Widgets/SBoxPanel.h"
+#include "Widgets/Text/STextBlock.h"
+
+SModalTaskWindow::SModalTaskWindow()
+{
+    AbortEvent = FPlatformProcess::GetSynchEventFromPool(true);
+    CloseEvent = FPlatformProcess::GetSynchEventFromPool(true);
+}
+
+SModalTaskWindow::~SModalTaskWindow()
+{
+    if(Thread != nullptr)
+    {
+        Thread->WaitForCompletion();
+        delete Thread;
+    }
+    FPlatformProcess::ReturnSynchEventToPool(CloseEvent);
+}
+
+void SModalTaskWindow::Construct(const FArguments& InArgs)
+{
+    SWindow::Construct(
+        SWindow::FArguments()
+        .Title(InArgs._Title)
+        .SizingRule(ESizingRule::Autosized)
+        .AutoCenter(EAutoCenter::PreferredWorkArea)
+        .HasCloseButton(false)
+        .SupportsMinimize(false)
+        .SupportsMaximize(false)
+        .MaxHeight(85.0f)
+        [
+            SNew(SVerticalBox)
+            + SVerticalBox::Slot()
+                .HAlign(HAlign_Center)
+                .VAlign(VAlign_Center)
+                .FillHeight(1.0f)
+                .Padding(FMargin(100.0f, 12.0f))
+                [
+                    SNew(STextBlock)
+                    .Text(InArgs._Message)
+                ]
+        ]
+    );
+
+    Task = InArgs._Task;
+
+    RegisterActiveTimer(0.1f, FWidgetActiveTimerDelegate::CreateSP(this, &SModalTaskWindow::OnTickTimer));
+    Thread = FRunnableThread::Create(this, TEXT("ModalTask"));
+}
+
+EActiveTimerReturnType SModalTaskWindow::OnTickTimer(double CurrentTime, float DeltaTime)
+{
+    if(CloseEvent->Wait(FTimespan::Zero()))
+    {
+        RequestDestroyWindow();
+        return EActiveTimerReturnType::Stop;
+    }
+    return EActiveTimerReturnType::Continue;
+}
+
+uint32 SModalTaskWindow::Run()
+{
+    Result = Task->Run(AbortEvent);
+    CloseEvent->Trigger();
+    return 0;
+}
+
+TSharedRef<FModalTaskResult> ExecuteModalTask(TSharedPtr<SWidget> Parent, TSharedRef<IModalTask> Task, const FText& InTitle, const FText& InMessage)
+{
+    TSharedRef<SModalTaskWindow> Window =
+        SNew(SModalTaskWindow)
+        .Title(InTitle)
+        .Message(InMessage)
+        .Task(Task);
+
+    FSlateApplication::Get().AddModalWindow(Window, TSharedPtr<SWidget>(), false);
+
+    return Window->Result.ToSharedRef();
+}
diff --git a/Engine/Source/Programs/SlateUGS/Private/Widgets/SModalTaskWindow.h b/Engine/Source/Programs/SlateUGS/Private/Widgets/SModalTaskWindow.h
new file mode 100644
--- /dev/null
+++ b/Engine/Source/Programs/SlateUGS/Private/Widgets/SModalTaskWindow.h
@@ -0,0 +1,37 @@
+// Copyright Epic Games, Inc. All Rights Reserved.
+
+#pragma once
+
+#include "CoreMinimal.h"
+#include "SlateCore.h"
+#include "HAL/Runnable.h"
+#include "HAL/RunnableThread.h"
+#include "UGSCore/ModalTask.h"
+
+class SModalTaskWindow : public SWindow, private FRunnable
+{
+public:
+    SLATE_BEGIN_ARGS(SModalTaskWindow) {}
+        SLATE_ARGUMENT(FText, Title)
+        SLATE_ARGUMENT(FText, Message)
+        SLATE_ARGUMENT(TSharedPtr<IModalTask>, Task)
+    SLATE_END_ARGS()
+
+    TSharedPtr<FModalTaskResult> Result;
+
+    SModalTaskWindow();
+    ~SModalTaskWindow();
+
+    void Construct(const FArguments& InArgs);
+    EActiveTimerReturnType OnTickTimer(double CurrentTime, float DeltaTime);
+
+private:
+    FEvent* AbortEvent;
+    FEvent* CloseEvent;
+    FRunnableThread* Thread;
+    TSharedPtr<IModalTask> Task;
+
+    virtual uint32 Run() override;
+};
+
+TSharedRef<FModalTaskResult> ExecuteModalTask(TSharedPtr<SWidget> Parent, TSharedRef<IModalTask> Task, const FText& InTitle, const FText& InMessage);
