Fixed error caused by leftover reference to deleted GameSyncController file

diff --git a/Engine/Source/Programs/SlateUGS/Private/SlateUGSApp.cpp b/Engine/Source/Programs/SlateUGS/Private/SlateUGSApp.cpp
--- a/Engine/Source/Programs/SlateUGS/Private/SlateUGSApp.cpp
+++ b/Engine/Source/Programs/SlateUGS/Private/SlateUGSApp.cpp
@@ -4,7 +4,6 @@
 #include "RequiredProgramMainCPPInclude.h"
 #include "Framework/Application/SlateApplication.h"
 
-#include "UGSCore/GameSyncController.h"
 #include "Widgets/SGameSyncTab.h"
 #include "Widgets/SWorkspaceWindow.h"
 #include "Widgets/SEmptyTab.h"
