Made modal loading window actually modal. Added touch ups to throbber animation

diff --git a/Engine/Source/Programs/SlateUGS/Private/Widgets/SModalTaskWindow.cpp b/Engine/Source/Programs/SlateUGS/Private/Widgets/SModalTaskWindow.cpp
--- a/Engine/Source/Programs/SlateUGS/Private/Widgets/SModalTaskWindow.cpp
+++ b/Engine/Source/Programs/SlateUGS/Private/Widgets/SModalTaskWindow.cpp
@@ -32,19 +32,21 @@ void SModalTaskWindow::Construct(const FArguments& InArgs)
 		.HasCloseButton(false)
 		.SupportsMinimize(false)
 		.SupportsMaximize(false)
-		.MaxHeight(85.0f)
+		.MaxHeight(140.0f)
 		[
 			SNew(SVerticalBox)
 			+SVerticalBox::Slot()
 			.HAlign(HAlign_Center)
 			.VAlign(VAlign_Center)
 			.FillHeight(1.0f)
-			.Padding(FMargin(100.0f, 12.0f))
+			.Padding(FMargin(100.0f, 8.0f))
 			[
 				SNew(STextBlock)
 				.Text(InArgs._Message)
 			]
 			+SVerticalBox::Slot()
+			.HAlign(HAlign_Center)
+			.VAlign(VAlign_Top)
 			[
 				SNew(SThrobber)
 			]
@@ -82,7 +84,7 @@ TSharedRef<FModalTaskResult> ExecuteModalTask(TSharedPtr<SWidget> Parent, TShare
 		.Message(InMessage)
 		.Task(Task);
 
-	FSlateApplication::Get().AddModalWindow(Window, TSharedPtr<SWidget>(), false);
+	FSlateApplication::Get().AddModalWindow(Window, Parent, false);
 
 	return Window->Result.ToSharedRef();
 }
