Add SLogWidget and add it in for Sync'ing only atm

Need one still for SEmptyTab when we are detecting settings but can be done later

diff --git a/Engine/Source/Programs/SlateUGS/Private/UGSTab.cpp b/Engine/Source/Programs/SlateUGS/Private/UGSTab.cpp
--- a/Engine/Source/Programs/SlateUGS/Private/UGSTab.cpp
+++ b/Engine/Source/Programs/SlateUGS/Private/UGSTab.cpp
@@ -2,11 +2,12 @@
 
 #include "UGSTab.h"
 
+#include "UGSCore/Utility.h"
 #include "UGSCore/BuildStep.h"
 #include "UGSCore/DetectProjectSettingsTask.h"
-#include "UGSCore/Utility.h"
 #include "UGSCore/UserSettings.h"
 #include "Widgets/SModalTaskWindow.h"
+#include "Widgets/SLogWidget.h"
 
 #define LOCTEXT_NAMESPACE "UGSTab"
 
@@ -210,6 +211,8 @@ namespace
 void UGSTab::SetupWorkspace()
 {
 	ProjectFileName = FUtility::GetPathWithCorrectCase(ProjectFileName);
+
+	// TODO likely should also log this on an Empty tab... so we can show logging info when we are loading things
 	TSharedRef<FDetectProjectSettingsTask> DetectSettings = MakeShared<FDetectProjectSettingsTask>(MakeShared<FPerforceConnection>(TEXT(""), TEXT(""), TEXT("")), ProjectFileName, MakeShared<FLineWriter>());
 
 	TSharedRef<FModalTaskResult> Result = ExecuteModalTask(TabWidget, DetectSettings, LOCTEXT("OpeningProjectTitle", "Opening Project"), LOCTEXT("OpeningProjectCaption", "Opening project, please wait..."));
@@ -250,40 +253,45 @@ void UGSTab::SetupWorkspace()
 	FString ProjectLogBaseName = DataFolder / FString::Printf(TEXT("%s@%s"), *PerforceClient->ClientName, *ClientKey.Replace(TEXT("/"), TEXT("$")));
 	FString TelemetryProjectIdentifier = FPerforceUtils::GetClientOrDepotDirectoryName(*SelectedProjectIdentifier);
 
-	Workspace = MakeShared<FWorkspace>(PerforceClient, BranchDirectoryName, ProjectFileName, BranchClientPath, SelectedClientFileName, CurrentChangeNumber, WorkspaceSettings->LastBuiltChangeNumber, TelemetryProjectIdentifier, MakeShared<FLineWriter>());
+	FString LogFileName = DataFolder / FPaths::GetPath(ProjectFileName) + TEXT(".sync.log");
+	GameSyncTabView->SetSyncLogLocation(LogFileName);
+
+	Workspace = MakeShared<FWorkspace>(PerforceClient, BranchDirectoryName, ProjectFileName, BranchClientPath, SelectedClientFileName, CurrentChangeNumber, WorkspaceSettings->LastBuiltChangeNumber, TelemetryProjectIdentifier, MakeShared<FLogWidgetTextWriter>(GameSyncTabView->GetSyncLog().ToSharedRef()));
 
 	// Todo:Move into a sync operation
-	// TArray<FString> CombinedSyncFilter = FUserSettings::GetCombinedSyncFilter(Workspace->GetSyncCategories(), Settings->SyncView, Settings->SyncExcludedCategories, WorkspaceSettings->SyncView, WorkspaceSettings->SyncExcludedCategories);
-
-	// // Options on what to do with workspace when updating it
-	// EWorkspaceUpdateOptions Options = EWorkspaceUpdateOptions::Sync | EWorkspaceUpdateOptions::SyncArchives | EWorkspaceUpdateOptions::GenerateProjectFiles;
-	// if(Settings->bAutoResolveConflicts)
-	// {
-	// 	Options |= EWorkspaceUpdateOptions::AutoResolveChanges;
-	// }
-	// if(Settings->bUseIncrementalBuilds)
-	// {
-	// 	Options |= EWorkspaceUpdateOptions::UseIncrementalBuilds;
-	// }
-	// if(Settings->bBuildAfterSync)
-	// {
-	// 	Options |= EWorkspaceUpdateOptions::Build;
-	// }
-	// if(Settings->bBuildAfterSync && Settings->bRunAfterSync)
-	// {
-	// 	Options |= EWorkspaceUpdateOptions::RunAfterSync; 
-	// }
-	// if(Settings->bOpenSolutionAfterSync)
-	// {
-	// 	Options |= EWorkspaceUpdateOptions::OpenSolutionAfterSync;
-	// }
+	/*
+	TArray<FString> CombinedSyncFilter = FUserSettings::GetCombinedSyncFilter(Workspace->GetSyncCategories(), Settings->SyncView, Settings->SyncExcludedCategories, WorkspaceSettings->SyncView, WorkspaceSettings->SyncExcludedCategories);
+
+	// Options on what to do with workspace when updating it
+	EWorkspaceUpdateOptions Options = EWorkspaceUpdateOptions::Sync | EWorkspaceUpdateOptions::SyncArchives | EWorkspaceUpdateOptions::GenerateProjectFiles;
+	if(Settings->bAutoResolveConflicts)
+	{
+		Options |= EWorkspaceUpdateOptions::AutoResolveChanges;
+	}
+	if(Settings->bUseIncrementalBuilds)
+	{
+		Options |= EWorkspaceUpdateOptions::UseIncrementalBuilds;
+	}
+	if(Settings->bBuildAfterSync)
+	{
+		Options |= EWorkspaceUpdateOptions::Build;
+	}
+	if(Settings->bBuildAfterSync && Settings->bRunAfterSync)
+	{
+		Options |= EWorkspaceUpdateOptions::RunAfterSync;
+	}
+	if(Settings->bOpenSolutionAfterSync)
+	{
+		Options |= EWorkspaceUpdateOptions::OpenSolutionAfterSync;
+	}
 
 	// // Hacking in the CL
-	// int ChangeNumber = 20658981;
-	// TSharedRef<FWorkspaceUpdateContext, ESPMode::ThreadSafe> Context = MakeShared<FWorkspaceUpdateContext, ESPMode::ThreadSafe>(ChangeNumber, Options, CombinedSyncFilter, GetDefaultBuildStepObjects(EditorTarget, Settings), ProjectSettings->BuildSteps, TSet<FGuid>(), GetWorkspaceVariables(DetectSettings));
+	int ChangeNumber = 20992206;
+	TSharedRef<FWorkspaceUpdateContext, ESPMode::ThreadSafe> Context = MakeShared<FWorkspaceUpdateContext, ESPMode::ThreadSafe>(ChangeNumber, Options, CombinedSyncFilter, GetDefaultBuildStepObjects(EditorTarget, Settings), ProjectSettings->BuildSteps, TSet<FGuid>(), GetWorkspaceVariables(DetectSettings));
 
 	// // Update the workspace with the Context!
-	// Workspace->Update(Context);
+	Workspace->Update(Context);
+	*/
 }
 
 #undef LOCTEXT_NAMESPACE
diff --git a/Engine/Source/Programs/SlateUGS/Private/Widgets/SGameSyncTab.cpp b/Engine/Source/Programs/SlateUGS/Private/Widgets/SGameSyncTab.cpp
--- a/Engine/Source/Programs/SlateUGS/Private/Widgets/SGameSyncTab.cpp
+++ b/Engine/Source/Programs/SlateUGS/Private/Widgets/SGameSyncTab.cpp
@@ -7,6 +7,8 @@
 #include "SPositiveActionButton.h"
 #include "SSimpleComboButton.h"
 #include "SSimpleButton.h"
+#include "Widgets/SLogWidget.h"
+#include "Widgets/Layout/SHeader.h"
 #include "Widgets/Testing/STestSuite.h"
 #include "Widgets/Docking/SDockTab.h"
 #include "Widgets/Colors/SSimpleGradient.h"
@@ -407,14 +409,35 @@ void SGameSyncTab::Construct(const FArguments& InArgs)
 			.ListItemsSource(&HordeBuilds)
 			.OnGenerateRow(this, &SGameSyncTab::GenerateHordeBuildTableRow)
 		]
-		// Console window
 		+SVerticalBox::Slot()
-		.Padding(20.0f, 5.0f)
 		[
-			SNew(STextBlock)
-			.Text(LOCTEXT("ConsoleArea", "Console Area"))
+			SNew(SVerticalBox)
+			+SVerticalBox::Slot()
+				.AutoHeight()
+				.Padding(0.0f, 4.0f, 0.0f, 8.0f)
+				[
+					SNew(SHeader)
+					[
+						SNew(STextBlock)
+						.Text(LOCTEXT("Log", "Log"))
+					]
+				]
+			+SVerticalBox::Slot()
+				[
+					SAssignNew(SyncLog, SLogWidget)
+				]
 		]
 	];
 }
 
+TSharedPtr<SLogWidget> SGameSyncTab::GetSyncLog() const
+{
+	return SyncLog;
+}
+
+void SGameSyncTab::SetSyncLogLocation(const FString& LogFileName)
+{
+    SyncLog->OpenFile(*LogFileName);
+}
+
 #undef LOCTEXT_NAMESPACE
diff --git a/Engine/Source/Programs/SlateUGS/Private/Widgets/SGameSyncTab.h b/Engine/Source/Programs/SlateUGS/Private/Widgets/SGameSyncTab.h
--- a/Engine/Source/Programs/SlateUGS/Private/Widgets/SGameSyncTab.h
+++ b/Engine/Source/Programs/SlateUGS/Private/Widgets/SGameSyncTab.h
@@ -5,6 +5,8 @@
 #include "Widgets/SCompoundWidget.h"
 #include "HordeBuildRowInfo.h"
 
+class SLogWidget;
+
 class SGameSyncTab : public SCompoundWidget
 {
 public:
@@ -14,12 +16,19 @@ public:
 
 	void Construct(const FArguments& InArgs);
 
+	// We need access to the SyncLog when creating the Workspace
+	// TODO: think of a better way to do this
+	TSharedPtr<SLogWidget> GetSyncLog() const;
+	void SetSyncLogLocation(const FString& LogFileName);
+
 private:
 	TSharedRef<ITableRow> GenerateHordeBuildTableRow(TSharedPtr<HordeBuildRowInfo> InItem, const TSharedRef<STableViewBase>& InOwnerTable);
 
-	TSharedPtr<SListView<TSharedPtr<HordeBuildRowInfo>>> HordeBuildsView; 
+	TSharedPtr<SListView<TSharedPtr<HordeBuildRowInfo>>> HordeBuildsView;
 	TArray<TSharedPtr<HordeBuildRowInfo>> HordeBuilds;
 
+	TSharedPtr<SLogWidget> SyncLog;
+
 	static constexpr float HordeBuildRowHorizontalPadding = 10.0f;
 	static constexpr float HordeBuildRowVerticalPadding = 2.5f;
 	static constexpr float HordeBuildRowExtraIconPadding = 10.0f;
diff --git a/Engine/Source/Programs/SlateUGS/Private/Widgets/SLogWidget.cpp b/Engine/Source/Programs/SlateUGS/Private/Widgets/SLogWidget.cpp
new file mode 100644
--- /dev/null
+++ b/Engine/Source/Programs/SlateUGS/Private/Widgets/SLogWidget.cpp
@@ -0,0 +1,188 @@
+// Copyright Epic Games, Inc. All Rights Reserved.
+
+#include "SLogWidget.h"
+#include "Framework/Text/SlateTextRun.h"
+
+BEGIN_SLATE_FUNCTION_BUILD_OPTIMIZATION
+
+#define LOCTEXT_NAMESPACE "UnrealGameSync"
+
+//// FLogWidgetTextLayoutMarshaller ////
+
+FLogWidgetTextLayoutMarshaller::FLogWidgetTextLayoutMarshaller()
+    : TextLayout(nullptr)
+{
+}
+
+FLogWidgetTextLayoutMarshaller::~FLogWidgetTextLayoutMarshaller()
+{
+}
+
+void FLogWidgetTextLayoutMarshaller::SetText(const FString& SourceString, FTextLayout& TargetTextLayout)
+{
+    TextLayout = &TargetTextLayout;
+
+    for(const TSharedRef<FString>& Line : Lines)
+    {
+        TextLayout->AddLine(FSlateTextLayout::FNewLineData(Line, TArray<TSharedRef<IRun>>()));
+    }
+}
+
+void FLogWidgetTextLayoutMarshaller::GetText(FString& TargetString, const FTextLayout& SourceTextLayout)
+{
+    SourceTextLayout.GetAsText(TargetString);
+}
+
+void FLogWidgetTextLayoutMarshaller::Clear()
+{
+    Lines.Empty();
+    MakeDirty();
+}
+
+void FLogWidgetTextLayoutMarshaller::AppendLine(const FString& Line)
+{
+    TSharedRef<FString> NewLine = MakeShared<FString>(Line);
+    Lines.Add(NewLine);
+
+    if(TextLayout)
+    {
+        // Remove the "default" line that's added for an empty text box.
+        if(Lines.Num() == 1)
+        {
+            TextLayout->ClearLines();
+        }
+
+        FTextBlockStyle Style = FTextBlockStyle::GetDefault();
+        Style.ColorAndOpacity = FSlateColor(FLinearColor::White);
+
+        TArray<TSharedRef<IRun>> Runs;
+        Runs.Add(FSlateTextRun::Create(FRunInfo(), NewLine, Style));
+        TextLayout->AddLine(FSlateTextLayout::FNewLineData(NewLine, Runs));
+    }
+}
+
+int32 FLogWidgetTextLayoutMarshaller::GetNumLines() const
+{
+    return Lines.Num();
+}
+
+//// SLogWidget ////
+
+SLogWidget::SLogWidget()
+    : bIsUserScrolled(false)
+    , LogWriter(nullptr)
+{
+}
+
+SLogWidget::~SLogWidget()
+{
+}
+
+void SLogWidget::Construct(const FArguments& InArgs)
+{
+    MessagesTextMarshaller = MakeShared<FLogWidgetTextLayoutMarshaller>();
+
+    ChildSlot
+    [
+        SNew(SBorder)
+        [
+             SAssignNew(MessagesTextBox, SMultiLineEditableTextBox)
+//              .Style(FEditorStyle::Get(), "Log.TextBox")
+//              .TextStyle(FEditorStyle::Get(), "Log.Normal")
+                .ForegroundColor(FLinearColor::Gray)
+                .Marshaller(MessagesTextMarshaller)
+                .IsReadOnly(true)
+                .AlwaysShowScrollbars(true)
+                .OnVScrollBarUserScrolled(this, &SLogWidget::OnScroll)
+//              .ContextMenuExtender(this, &SLogWidget::ExtendTextBoxMenu);
+        ]
+    ];
+
+    RegisterActiveTimer(0.03f, FWidgetActiveTimerDelegate::CreateSP(this, &SLogWidget::OnTimerElapsed));
+}
+
+bool SLogWidget::OpenFile(const TCHAR* NewLogFileName)
+{
+    CloseFile();
+    Clear();
+
+    TArray<FString> InitialLines;
+    if(FFileHelper::LoadFileToStringArray(InitialLines, NewLogFileName))
+    {
+        for(const FString& InitialLine : InitialLines)
+        {
+            AppendLine(InitialLine);
+        }
+    }
+
+    LogWriter = IFileManager::Get().CreateFileWriter(NewLogFileName, FILEWRITE_Append);
+    if(LogWriter == nullptr)
+    {
+        return false;
+    }
+
+    LogFileName = NewLogFileName;
+    return true;
+}
+
+void SLogWidget::CloseFile()
+{
+    if(LogWriter != nullptr)
+    {
+        delete LogWriter;
+        LogWriter = nullptr;
+    }
+}
+
+void SLogWidget::Clear()
+{
+    if(LogWriter != nullptr)
+    {
+        delete LogWriter;
+        LogWriter = IFileManager::Get().CreateFileWriter(*LogFileName, 0);
+    }
+    MessagesTextMarshaller->Clear();
+}
+
+void SLogWidget::ScrollToEnd()
+{
+    MessagesTextBox->ScrollTo(FTextLocation(MessagesTextMarshaller->GetNumLines() - 1));
+    bIsUserScrolled = false;
+}
+
+void SLogWidget::AppendLine(const FString& Line)
+{
+    FScopeLock Lock(&CriticalSection);
+    QueuedLines.Add(Line);
+}
+
+void SLogWidget::OnScroll(float ScrollOffset)
+{
+    bIsUserScrolled = ScrollOffset < 1.0 && !FMath::IsNearlyEqual(ScrollOffset, 1.0f);
+}
+
+EActiveTimerReturnType SLogWidget::OnTimerElapsed(double CurrentTime, float DeltaTime)
+{
+    FScopeLock Lock(&CriticalSection);
+    for(const FString& QueuedLine : QueuedLines)
+    {
+        if(LogWriter != nullptr)
+        {
+            LogWriter->Serialize((void*)*QueuedLine, sizeof(TCHAR) * QueuedLine.Len());
+
+            const TCHAR LineTerminator[] = LINE_TERMINATOR;
+            LogWriter->Serialize((void*)LineTerminator, sizeof(LineTerminator));
+        }
+        MessagesTextMarshaller->AppendLine(QueuedLine);
+    }
+    if(!bIsUserScrolled)
+    {
+        ScrollToEnd();
+    }
+    QueuedLines.Empty();
+    return EActiveTimerReturnType::Continue;
+}
+
+#undef LOCTEXT_NAMESPACE
+
+END_SLATE_FUNCTION_BUILD_OPTIMIZATION
diff --git a/Engine/Source/Programs/SlateUGS/Private/Widgets/SLogWidget.h b/Engine/Source/Programs/SlateUGS/Private/Widgets/SLogWidget.h
new file mode 100644
--- /dev/null
+++ b/Engine/Source/Programs/SlateUGS/Private/Widgets/SLogWidget.h
@@ -0,0 +1,75 @@
+// Copyright Epic Games, Inc. All Rights Reserved.
+
+#pragma once
+
+#include "SlateCore.h"
+#include "UGSCore/OutputAdapters.h"
+#include "Framework/Text/BaseTextLayoutMarshaller.h"
+#include "Widgets/Input/SMultiLineEditableTextBox.h"
+
+class FLogWidgetTextLayoutMarshaller : public FBaseTextLayoutMarshaller
+{
+public:
+    FLogWidgetTextLayoutMarshaller();
+    virtual ~FLogWidgetTextLayoutMarshaller();
+
+    // ITextLayoutMarshaller
+    virtual void SetText(const FString& SourceString, FTextLayout& TargetTextLayout) override;
+    virtual void GetText(FString& TargetString, const FTextLayout& SourceTextLayout) override;
+
+    void Clear();
+    void AppendLine(const FString& Line);
+    int32 GetNumLines() const;
+
+private:
+    TArray<TSharedRef<FString>> Lines;
+    FTextLayout* TextLayout;
+};
+
+class SLogWidget : public SCompoundWidget
+{
+public:
+    SLATE_BEGIN_ARGS(SLogWidget) { }
+    SLATE_END_ARGS()
+
+    SLogWidget();
+    ~SLogWidget();
+
+    void Construct( const FArguments& InArgs );
+
+    bool OpenFile(const TCHAR* NewLogFileName);
+    void CloseFile();
+
+    void Clear();
+    void ScrollToEnd();
+    void AppendLine(const FString& Line);
+
+protected:
+	FString LogFileName;
+    FCriticalSection CriticalSection;
+    TArray<FString> QueuedLines;
+    TSharedPtr<FLogWidgetTextLayoutMarshaller> MessagesTextMarshaller;
+    TSharedPtr<SMultiLineEditableTextBox> MessagesTextBox;
+    bool bIsUserScrolled;
+    FArchive* LogWriter;
+
+    void OnScroll(float ScrollOffset);
+    EActiveTimerReturnType OnTimerElapsed(double CurrentTime, float DeltaTime);
+};
+
+class FLogWidgetTextWriter : public FLineBasedTextWriter
+{
+public:
+    FLogWidgetTextWriter(const TSharedRef<SLogWidget>& InLogWidget)
+        : LogWidget(InLogWidget)
+    {
+    }
+
+    virtual void FlushLine(const FString& Line) override
+    {
+        LogWidget->AppendLine(Line);
+    }
+
+private:
+    TSharedRef<SLogWidget> LogWidget;
+};
