Set the style of the log window to be of a log type

diff --git a/Engine/Source/Programs/SlateUGS/Private/Widgets/SLogWidget.cpp b/Engine/Source/Programs/SlateUGS/Private/Widgets/SLogWidget.cpp
--- a/Engine/Source/Programs/SlateUGS/Private/Widgets/SLogWidget.cpp
+++ b/Engine/Source/Programs/SlateUGS/Private/Widgets/SLogWidget.cpp
@@ -10,7 +10,7 @@ BEGIN_SLATE_FUNCTION_BUILD_OPTIMIZATION
 //// FLogWidgetTextLayoutMarshaller ////
 
 FLogWidgetTextLayoutMarshaller::FLogWidgetTextLayoutMarshaller()
-    : TextLayout(nullptr)
+	: TextLayout(nullptr)
 {
 }
 
@@ -20,57 +20,57 @@ FLogWidgetTextLayoutMarshaller::~FLogWidgetTextLayoutMarshaller()
 
 void FLogWidgetTextLayoutMarshaller::SetText(const FString& SourceString, FTextLayout& TargetTextLayout)
 {
-    TextLayout = &TargetTextLayout;
+	TextLayout = &TargetTextLayout;
 
-    for(const TSharedRef<FString>& Line : Lines)
-    {
-        TextLayout->AddLine(FSlateTextLayout::FNewLineData(Line, TArray<TSharedRef<IRun>>()));
-    }
+	for(const TSharedRef<FString>& Line : Lines)
+	{
+		TextLayout->AddLine(FSlateTextLayout::FNewLineData(Line, TArray<TSharedRef<IRun>>()));
+	}
 }
 
 void FLogWidgetTextLayoutMarshaller::GetText(FString& TargetString, const FTextLayout& SourceTextLayout)
 {
-    SourceTextLayout.GetAsText(TargetString);
+	SourceTextLayout.GetAsText(TargetString);
 }
 
 void FLogWidgetTextLayoutMarshaller::Clear()
 {
-    Lines.Empty();
-    MakeDirty();
+	Lines.Empty();
+	MakeDirty();
 }
 
 void FLogWidgetTextLayoutMarshaller::AppendLine(const FString& Line)
 {
-    TSharedRef<FString> NewLine = MakeShared<FString>(Line);
-    Lines.Add(NewLine);
+	TSharedRef<FString> NewLine = MakeShared<FString>(Line);
+	Lines.Add(NewLine);
 
-    if(TextLayout)
-    {
-        // Remove the "default" line that's added for an empty text box.
-        if(Lines.Num() == 1)
-        {
-            TextLayout->ClearLines();
-        }
+	if(TextLayout)
+	{
+		// Remove the "default" line that's added for an empty text box.
+		if(Lines.Num() == 1)
+		{
+			TextLayout->ClearLines();
+		}
 
-        FTextBlockStyle Style = FTextBlockStyle::GetDefault();
-        Style.ColorAndOpacity = FSlateColor(FLinearColor::White);
+		FTextBlockStyle Style = FTextBlockStyle::GetDefault();
+		Style.ColorAndOpacity = FSlateColor(FLinearColor::White);
 
-        TArray<TSharedRef<IRun>> Runs;
-        Runs.Add(FSlateTextRun::Create(FRunInfo(), NewLine, Style));
-        TextLayout->AddLine(FSlateTextLayout::FNewLineData(NewLine, Runs));
-    }
+		TArray<TSharedRef<IRun>> Runs;
+		Runs.Add(FSlateTextRun::Create(FRunInfo(), NewLine, Style));
+		TextLayout->AddLine(FSlateTextLayout::FNewLineData(NewLine, Runs));
+	}
 }
 
 int32 FLogWidgetTextLayoutMarshaller::GetNumLines() const
 {
-    return Lines.Num();
+	return Lines.Num();
 }
 
 //// SLogWidget ////
 
 SLogWidget::SLogWidget()
-    : bIsUserScrolled(false)
-    , LogWriter(nullptr)
+	: bIsUserScrolled(false)
+	, LogWriter(nullptr)
 {
 }
 
@@ -80,107 +80,106 @@ SLogWidget::~SLogWidget()
 
 void SLogWidget::Construct(const FArguments& InArgs)
 {
-    MessagesTextMarshaller = MakeShared<FLogWidgetTextLayoutMarshaller>();
+	MessagesTextMarshaller = MakeShared<FLogWidgetTextLayoutMarshaller>();
 
-    ChildSlot
-    [
-        SNew(SBorder)
-        [
-             SAssignNew(MessagesTextBox, SMultiLineEditableTextBox)
-//              .Style(FEditorStyle::Get(), "Log.TextBox")
-//              .TextStyle(FEditorStyle::Get(), "Log.Normal")
-                .ForegroundColor(FLinearColor::Gray)
-                .Marshaller(MessagesTextMarshaller)
-                .IsReadOnly(true)
-                .AlwaysShowScrollbars(true)
-                .OnVScrollBarUserScrolled(this, &SLogWidget::OnScroll)
-//              .ContextMenuExtender(this, &SLogWidget::ExtendTextBoxMenu);
-        ]
-    ];
+	ChildSlot
+	[
+		SNew(SBorder)
+		[
+			 SAssignNew(MessagesTextBox, SMultiLineEditableTextBox)
+				.Style(FAppStyle::Get(), "Log.TextBox")
+				.TextStyle(FAppStyle::Get(), "Log.Normal")
+				.ForegroundColor(FLinearColor::Gray)
+				.Marshaller(MessagesTextMarshaller)
+				.IsReadOnly(true)
+				.AlwaysShowScrollbars(true)
+				.OnVScrollBarUserScrolled(this, &SLogWidget::OnScroll)
+		]
+	];
 
-    RegisterActiveTimer(0.03f, FWidgetActiveTimerDelegate::CreateSP(this, &SLogWidget::OnTimerElapsed));
+	RegisterActiveTimer(0.03f, FWidgetActiveTimerDelegate::CreateSP(this, &SLogWidget::OnTimerElapsed));
 }
 
 bool SLogWidget::OpenFile(const TCHAR* NewLogFileName)
 {
-    CloseFile();
-    Clear();
+	CloseFile();
+	Clear();
 
-    TArray<FString> InitialLines;
-    if(FFileHelper::LoadFileToStringArray(InitialLines, NewLogFileName))
-    {
-        for(const FString& InitialLine : InitialLines)
-        {
-            AppendLine(InitialLine);
-        }
-    }
+	TArray<FString> InitialLines;
+	if(FFileHelper::LoadFileToStringArray(InitialLines, NewLogFileName))
+	{
+		for(const FString& InitialLine : InitialLines)
+		{
+			AppendLine(InitialLine);
+		}
+	}
 
-    LogWriter = IFileManager::Get().CreateFileWriter(NewLogFileName, FILEWRITE_Append);
-    if(LogWriter == nullptr)
-    {
-        return false;
-    }
+	LogWriter = IFileManager::Get().CreateFileWriter(NewLogFileName, FILEWRITE_Append);
+	if(LogWriter == nullptr)
+	{
+		return false;
+	}
 
-    LogFileName = NewLogFileName;
-    return true;
+	LogFileName = NewLogFileName;
+	return true;
 }
 
 void SLogWidget::CloseFile()
 {
-    if(LogWriter != nullptr)
-    {
-        delete LogWriter;
-        LogWriter = nullptr;
-    }
+	if(LogWriter != nullptr)
+	{
+		delete LogWriter;
+		LogWriter = nullptr;
+	}
 }
 
 void SLogWidget::Clear()
 {
-    if(LogWriter != nullptr)
-    {
-        delete LogWriter;
-        LogWriter = IFileManager::Get().CreateFileWriter(*LogFileName, 0);
-    }
-    MessagesTextMarshaller->Clear();
+	if(LogWriter != nullptr)
+	{
+		delete LogWriter;
+		LogWriter = IFileManager::Get().CreateFileWriter(*LogFileName, 0);
+	}
+	MessagesTextMarshaller->Clear();
 }
 
 void SLogWidget::ScrollToEnd()
 {
-    MessagesTextBox->ScrollTo(FTextLocation(MessagesTextMarshaller->GetNumLines() - 1));
-    bIsUserScrolled = false;
+	MessagesTextBox->ScrollTo(FTextLocation(MessagesTextMarshaller->GetNumLines() - 1));
+	bIsUserScrolled = false;
 }
 
 void SLogWidget::AppendLine(const FString& Line)
 {
-    FScopeLock Lock(&CriticalSection);
-    QueuedLines.Add(Line);
+	FScopeLock Lock(&CriticalSection);
+	QueuedLines.Add(Line);
 }
 
 void SLogWidget::OnScroll(float ScrollOffset)
 {
-    bIsUserScrolled = ScrollOffset < 1.0 && !FMath::IsNearlyEqual(ScrollOffset, 1.0f);
+	bIsUserScrolled = ScrollOffset < 1.0 && !FMath::IsNearlyEqual(ScrollOffset, 1.0f);
 }
 
 EActiveTimerReturnType SLogWidget::OnTimerElapsed(double CurrentTime, float DeltaTime)
 {
-    FScopeLock Lock(&CriticalSection);
-    for(const FString& QueuedLine : QueuedLines)
-    {
-        if(LogWriter != nullptr)
-        {
-            LogWriter->Serialize((void*)*QueuedLine, sizeof(TCHAR) * QueuedLine.Len());
-
-            const TCHAR LineTerminator[] = LINE_TERMINATOR;
-            LogWriter->Serialize((void*)LineTerminator, sizeof(LineTerminator));
-        }
-        MessagesTextMarshaller->AppendLine(QueuedLine);
-    }
-    if(!bIsUserScrolled)
-    {
-        ScrollToEnd();
-    }
-    QueuedLines.Empty();
-    return EActiveTimerReturnType::Continue;
+	FScopeLock Lock(&CriticalSection);
+	for(const FString& QueuedLine : QueuedLines)
+	{
+		if(LogWriter != nullptr)
+		{
+			LogWriter->Serialize((void*)*QueuedLine, sizeof(TCHAR) * QueuedLine.Len());
+
+			const TCHAR LineTerminator[] = LINE_TERMINATOR;
+			LogWriter->Serialize((void*)LineTerminator, sizeof(LineTerminator));
+		}
+		MessagesTextMarshaller->AppendLine(QueuedLine);
+	}
+	if(!bIsUserScrolled)
+	{
+		ScrollToEnd();
+	}
+	QueuedLines.Empty();
+	return EActiveTimerReturnType::Continue;
 }
 
 #undef LOCTEXT_NAMESPACE
