SlateUGS - Find p4 exe for 32 & 64 bit windows, fix terminators

diff --git a/Engine/Source/Programs/SlateUGS/Private/UGSCore/Perforce.cpp b/Engine/Source/Programs/SlateUGS/Private/UGSCore/Perforce.cpp
--- a/Engine/Source/Programs/SlateUGS/Private/UGSCore/Perforce.cpp
+++ b/Engine/Source/Programs/SlateUGS/Private/UGSCore/Perforce.cpp
@@ -36,11 +36,76 @@ FString GetTempFileName()
 	return FPaths::EngineIntermediateDir() / FGuid::NewGuid().ToString();
 }
 
+struct FPerforceExe
+{
 #if PLATFORM_WINDOWS
-const TCHAR *PerforceExe = TEXT("C:\\Program Files (x86)\\Perforce\\p4.exe");
+	FString PerforceExe = TEXT("C:\\Program Files (x86)\\Perforce\\p4.exe");
 #else
-const TCHAR *PerforceExe = TEXT("p4");
+	FString PerforceExe = TEXT("p4");
 #endif
+	bool bValidP4 = false;
+
+	inline int InnerRunCommand(const FString& CommandLine, TArray<FString>& OutLines)
+	{
+		uint32 ProcId;
+		void* ReadPipe = nullptr;
+		void* WritePipe = nullptr;
+		FPlatformProcess::CreatePipe(ReadPipe, WritePipe);
+
+		FProcHandle P4Proc = FPlatformProcess::CreateProc(*PerforceExe, *CommandLine, false, false, false, &ProcId, 0, nullptr, WritePipe, ReadPipe);
+
+		FString P4Output;
+		FString LatestOutput = FPlatformProcess::ReadPipe(ReadPipe);
+		while (FPlatformProcess::IsProcRunning(P4Proc) || !LatestOutput.IsEmpty())
+		{
+			P4Output += LatestOutput;
+			LatestOutput = FPlatformProcess::ReadPipe(ReadPipe);
+			FPlatformProcess::Sleep(0);
+		}
+
+		FPlatformProcess::ClosePipe(ReadPipe, WritePipe);
+
+		// TODO check if p4 spits out \n on Windows
+		OutLines = Split(P4Output, LINE_TERMINATOR);
+
+		int ExitCode = -1;
+		bool GotReturnCode = FPlatformProcess::GetProcReturnCode(P4Proc, &ExitCode);
+
+		if (GotReturnCode)
+		{
+			return ExitCode;
+		}
+
+		return -1;
+	}
+
+	int RunCommand(const FString& CommandLine, TArray<FString>& OutLines)
+	{
+		if (bValidP4)
+		{
+			return InnerRunCommand(CommandLine, OutLines);
+		}
+
+		return -1;
+	}
+
+	FPerforceExe()
+	{
+#if PLATFORM_WINDOWS
+		TArray<FString> Lines;
+		if (RunCommand(TEXT(""), Lines) == -1)
+		{
+			PerforceExe.ReplaceInline(TEXT(" (x86)"), TEXT(""));
+			bValidP4 = (RunCommand(TEXT(""), Lines) == -1);
+		}
+		else
+#endif
+		{
+			bValidP4 = true;
+		}
+	}
+};
+static FPerforceExe GPerforceExe;
 
 //// FPerforceChangeSummary ////
 
@@ -987,40 +1052,6 @@ bool FPerforceConnection::RunCommand(const FString& CommandLine, TArray<FString>
 	return RunCommand(CommandLine, EPerforceOutputChannel::Info, OutLines, Options, AbortEvent, Log);
 }
 
-int RunPerforceCommand(const FString& CommandLine, TArray<FString>& OutLines)
-{
-	uint32 ProcId;
-	void* ReadPipe = nullptr;
-    void* WritePipe = nullptr;
-    FPlatformProcess::CreatePipe(ReadPipe, WritePipe);
-
-	FProcHandle P4Proc = FPlatformProcess::CreateProc(PerforceExe, *CommandLine, false, false, false, &ProcId, 0, nullptr, WritePipe, ReadPipe);
-
-	FString P4Output;
-    FString LatestOutput = FPlatformProcess::ReadPipe(ReadPipe);
-    while (FPlatformProcess::IsProcRunning(P4Proc) || !LatestOutput.IsEmpty())
-    {
-        P4Output += LatestOutput;
-        LatestOutput = FPlatformProcess::ReadPipe(ReadPipe);
-        FPlatformProcess::Sleep(0);
-    }
-
-    FPlatformProcess::ClosePipe(ReadPipe, WritePipe);
-
-	// TODO check if p4 spits out \n on Windows
-	OutLines = Split(P4Output, TEXT("\n"));
-
-	int ExitCode = -1;
-	bool GotReturnCode = FPlatformProcess::GetProcReturnCode(P4Proc, &ExitCode);
-
-	if (GotReturnCode)
-	{
-		return ExitCode;
-	}
-
-	return -1;
-}
-
 bool FPerforceConnection::RunCommand(const FString& CommandLine, EPerforceOutputChannel Channel, TArray<FString>& OutLines, ECommandOptions Options, FEvent* AbortEvent, FOutputDevice& Log) const
 {
 	FString FullCommandLine = GetFullCommandLine(CommandLine, Options);
@@ -1028,7 +1059,7 @@ bool FPerforceConnection::RunCommand(const FString& CommandLine, EPerforceOutput
 
 	TArray<FString> RawOutputLines;
 
-	int ExitCode = RunPerforceCommand(FullCommandLine, RawOutputLines);
+	int ExitCode = GPerforceExe.RunCommand(FullCommandLine, RawOutputLines);
 	if (ExitCode != 0 && !EnumHasAnyFlags(Options, ECommandOptions::IgnoreExitCode))
 	{
 		return false;
@@ -1058,7 +1089,7 @@ bool FPerforceConnection::RunCommand(const FString& CommandLine, const TCHAR* In
 
 	TArray<FString> RawOutputLines;
 
-	int ExitCode = RunPerforceCommand(FullCommandLine, RawOutputLines);
+	int ExitCode = GPerforceExe.RunCommand(FullCommandLine, RawOutputLines);
 
 	// TODO check this vs the old way as things *may* have changed. Before they were handling each line as it was coming
 	// back from the process. here we just spin + collect all the output. Then try to process all the output, then if the
@@ -1316,7 +1347,7 @@ bool FPerforceConnection::RunCommandWithBinaryOutput(const FString& CommandLine,
     void* WritePipe = nullptr;
     FPlatformProcess::CreatePipe(ReadPipe, WritePipe);
 
-	FProcHandle P4Proc = FPlatformProcess::CreateProc(PerforceExe, *FullCommandLine, false, false, false, &ProcId, 0, nullptr, WritePipe, ReadPipe);
+	FProcHandle P4Proc = FPlatformProcess::CreateProc(*GPerforceExe.PerforceExe, *FullCommandLine, false, false, false, &ProcId, 0, nullptr, WritePipe, ReadPipe);
 
 	/*
 	FString P4Output;
