Actually get the Latest CL

diff --git a/Engine/Source/Programs/SlateUGS/Private/UGSCore/Perforce.cpp b/Engine/Source/Programs/SlateUGS/Private/UGSCore/Perforce.cpp
--- a/Engine/Source/Programs/SlateUGS/Private/UGSCore/Perforce.cpp
+++ b/Engine/Source/Programs/SlateUGS/Private/UGSCore/Perforce.cpp
@@ -926,6 +926,23 @@ bool FPerforceConnection::Sync(const TArray<FString>& DepotPaths, int ChangeNumb
 	return RunCommand(CommandLine, nullptr, [&Parser, &OutTamperedFiles, &Log](const FPerforceOutputLine& Line){ return FilterSyncOutput(Line, Parser, OutTamperedFiles, Log); }, ECommandOptions::NoFailOnErrors | ECommandOptions::IgnoreFilesUpToDateError | ECommandOptions::IgnoreExitCode, AbortEvent, Log);
 }
 
+bool FPerforceConnection::LatestChangeList(int& OutChangeList, FEvent* AbortEvent, FOutputDevice& Log) const
+{
+	TArray<TMap<FString, FString>> OutLatestChange;
+	RunCommand(TEXT("changes -m 1"), OutLatestChange, ECommandOptions::None, AbortEvent, Log);
+
+	if (OutLatestChange.Num() == 1)
+	{
+		FString StringChangeList;
+		if (TryGetValue(OutLatestChange[0], TEXT("change"), StringChangeList))
+		{
+			return FUtility::TryParse(*StringChangeList, OutChangeList);
+		}
+	}
+
+	return false;
+}
+
 bool FPerforceConnection::FilterSyncOutput(const FPerforceOutputLine& Line, FPerforceTagRecordParser& Parser, TArray<FString>& OutTamperedFiles, FOutputDevice& Log)
 {
 	if(Line.Channel == EPerforceOutputChannel::TaggedInfo)
diff --git a/Engine/Source/Programs/SlateUGS/Private/UGSCore/Perforce.h b/Engine/Source/Programs/SlateUGS/Private/UGSCore/Perforce.h
--- a/Engine/Source/Programs/SlateUGS/Private/UGSCore/Perforce.h
+++ b/Engine/Source/Programs/SlateUGS/Private/UGSCore/Perforce.h
@@ -227,6 +227,7 @@ public:
 	bool Stat(const FString& Filter, TArray<FPerforceFileRecord>& OutFileRecords, FEvent* AbortEvent, FOutputDevice& Log) const;
 	bool Sync(const FString& Filter, FEvent* AbortEvent, FOutputDevice& Log) const;
 	bool Sync(const TArray<FString>& DepotPaths, int ChangeNumber, TFunction<void(const FPerforceFileRecord&)> SyncOutput, TArray<FString>& OutTamperedFiles, const FPerforceSyncOptions* Options, FEvent* AbortEvent, FOutputDevice& Log) const;
+	bool LatestChangeList(int& OutChangeList, FEvent* AbortEvent, FOutputDevice& Log) const;
 
 private:
 	static bool FilterSyncOutput(const FPerforceOutputLine& Line, FPerforceTagRecordParser& Parser, TArray<FString>& OutTamperedFiles, FOutputDevice& Log);
diff --git a/Engine/Source/Programs/SlateUGS/Private/UGSTab.cpp b/Engine/Source/Programs/SlateUGS/Private/UGSTab.cpp
--- a/Engine/Source/Programs/SlateUGS/Private/UGSTab.cpp
+++ b/Engine/Source/Programs/SlateUGS/Private/UGSTab.cpp
@@ -208,19 +208,25 @@ bool UGSTab::OnWorkspaceChosen(const FString& Project)
 
 void UGSTab::OnSyncLatest()
 {
-	// Hacking in the CL
-	int ChangeNumber = 20992726;
-	TSharedRef<FWorkspaceUpdateContext, ESPMode::ThreadSafe> Context = MakeShared<FWorkspaceUpdateContext, ESPMode::ThreadSafe>(
-		ChangeNumber,
-		Options,
-		CombinedSyncFilter,
-		GetDefaultBuildStepObjects(DetectSettings->NewProjectEditorTarget, UserSettings),
-		ProjectSettings->BuildSteps,
-		TSet<FGuid>(),
-		GetWorkspaceVariables(DetectSettings));
-
-	// Update the workspace with the Context!
-	Workspace->Update(Context);
+	int ChangeNumber = -1;
+	FEvent* AbortEvent = FPlatformProcess::GetSynchEventFromPool(true);
+
+	if (PerforceClient->LatestChangeList(ChangeNumber, AbortEvent, MakeShared<FLogWidgetTextWriter>(GameSyncTabView->GetSyncLog().ToSharedRef()).Get()))
+	{
+		TSharedRef<FWorkspaceUpdateContext, ESPMode::ThreadSafe> Context = MakeShared<FWorkspaceUpdateContext, ESPMode::ThreadSafe>(
+			ChangeNumber,
+			Options,
+			CombinedSyncFilter,
+			GetDefaultBuildStepObjects(DetectSettings->NewProjectEditorTarget, UserSettings),
+			ProjectSettings->BuildSteps,
+			TSet<FGuid>(),
+			GetWorkspaceVariables(DetectSettings));
+
+		// Update the workspace with the Context!
+		Workspace->Update(Context);
+	}
+
+	FPlatformProcess::ReturnSynchEventToPool(AbortEvent);
 }
 
 void UGSTab::SetupWorkspace()
