Add log category

diff --git a/Engine/Source/Programs/SlateUGS/Private/SlateUGSApp.cpp b/Engine/Source/Programs/SlateUGS/Private/SlateUGSApp.cpp
--- a/Engine/Source/Programs/SlateUGS/Private/SlateUGSApp.cpp
+++ b/Engine/Source/Programs/SlateUGS/Private/SlateUGSApp.cpp
@@ -1,6 +1,7 @@
 // Copyright Epic Games, Inc. All Rights Reserved.
 
 #include "SlateUGSApp.h"
+#include "UGSLog.h"
 #include "RequiredProgramMainCPPInclude.h"
 #include "Framework/Application/SlateApplication.h"
 
@@ -11,6 +12,7 @@
 #include "UGSTabManager.h"
 
 IMPLEMENT_APPLICATION(SlateUGS, "SlateUGS");
+DEFINE_LOG_CATEGORY(LogSlateUGS);
 
 #define LOCTEXT_NAMESPACE "SlateUGS"
 
diff --git a/Engine/Source/Programs/SlateUGS/Private/UGSLog.h b/Engine/Source/Programs/SlateUGS/Private/UGSLog.h
new file mode 100644
--- /dev/null
+++ b/Engine/Source/Programs/SlateUGS/Private/UGSLog.h
@@ -0,0 +1,7 @@
+// Copyright Epic Games, Inc. All Rights Reserved.
+
+#pragma once
+
+#include "CoreMinimal.h"
+
+DECLARE_LOG_CATEGORY_EXTERN(LogSlateUGS, Verbose, All);
diff --git a/Engine/Source/Programs/SlateUGS/Private/UGSTab.cpp b/Engine/Source/Programs/SlateUGS/Private/UGSTab.cpp
--- a/Engine/Source/Programs/SlateUGS/Private/UGSTab.cpp
+++ b/Engine/Source/Programs/SlateUGS/Private/UGSTab.cpp
@@ -2,6 +2,7 @@
 
 #include "UGSTab.h"
 
+#include "UGSLog.h"
 #include "UGSCore/Utility.h"
 #include "UGSCore/BuildStep.h"
 #include "Widgets/SModalTaskWindow.h"
@@ -114,7 +115,7 @@ namespace
 	public:
 		virtual void FlushLine(const FString& Line) override
 		{
-			printf("%s\n", TCHAR_TO_ANSI(*Line));
+			UE_LOG(LogSlateUGS, Log, TEXT("%s"), *Line);
 		}
 	};
 
