Use the PerforceExe name vs hardcoding p4.exe

diff --git a/Engine/Source/Programs/SlateUGS/Private/UGSCore/Perforce.cpp b/Engine/Source/Programs/SlateUGS/Private/UGSCore/Perforce.cpp
--- a/Engine/Source/Programs/SlateUGS/Private/UGSCore/Perforce.cpp
+++ b/Engine/Source/Programs/SlateUGS/Private/UGSCore/Perforce.cpp
@@ -1072,7 +1072,7 @@ bool FPerforceConnection::RunCommand(const FString& CommandLine, TArray<FString>
 bool FPerforceConnection::RunCommand(const FString& CommandLine, EPerforceOutputChannel Channel, TArray<FString>& OutLines, ECommandOptions Options, FEvent* AbortEvent, FOutputDevice& Log) const
 {
 	FString FullCommandLine = GetFullCommandLine(CommandLine, Options);
-	Log.Logf(TEXT("p4> p4.exe %s"), *FullCommandLine);
+	Log.Logf(TEXT("p4> %s %s"), *FPaths::GetCleanFilename(GPerforceExe.PerforceExe), *FullCommandLine);
 
 	TArray<FString> RawOutputLines;
 
@@ -1102,7 +1102,7 @@ bool FPerforceConnection::RunCommand(const FString& CommandLine, EPerforceOutput
 bool FPerforceConnection::RunCommand(const FString& CommandLine, const TCHAR* Input, TFunction<bool(const FPerforceOutputLine&)> HandleOutput, ECommandOptions Options, FEvent* AbortEvent, FOutputDevice& Log) const
 {
 	FString FullCommandLine = GetFullCommandLine(CommandLine, Options);
-	Log.Logf(TEXT("p4> p4.exe %s"), *FullCommandLine);
+	Log.Logf(TEXT("p4> %s %s"), *FPaths::GetCleanFilename(GPerforceExe.PerforceExe), *FullCommandLine);
 
 	TArray<FString> RawOutputLines;
 
