Fixed duplicated logging

diff --git a/Engine/Source/Programs/SlateUGS/Private/Widgets/SLogWidget.cpp b/Engine/Source/Programs/SlateUGS/Private/Widgets/SLogWidget.cpp
--- a/Engine/Source/Programs/SlateUGS/Private/Widgets/SLogWidget.cpp
+++ b/Engine/Source/Programs/SlateUGS/Private/Widgets/SLogWidget.cpp
@@ -114,7 +114,7 @@ bool SLogWidget::OpenFile(const TCHAR* NewLogFileName)
 		}
 	}
 
-	LogWriter = IFileManager::Get().CreateFileWriter(NewLogFileName, FILEWRITE_Append);
+	LogWriter = IFileManager::Get().CreateFileWriter(NewLogFileName);
 	if(LogWriter == nullptr)
 	{
 		return false;
@@ -138,7 +138,7 @@ void SLogWidget::Clear()
 	if(LogWriter != nullptr)
 	{
 		delete LogWriter;
-		LogWriter = IFileManager::Get().CreateFileWriter(*LogFileName, 0);
+		LogWriter = IFileManager::Get().CreateFileWriter(*LogFileName);
 	}
 	MessagesTextMarshaller->Clear();
 }
