Clean up AbortEvent on dtor

diff --git a/Engine/Source/Programs/SlateUGS/Private/Widgets/SModalTaskWindow.cpp b/Engine/Source/Programs/SlateUGS/Private/Widgets/SModalTaskWindow.cpp
--- a/Engine/Source/Programs/SlateUGS/Private/Widgets/SModalTaskWindow.cpp
+++ b/Engine/Source/Programs/SlateUGS/Private/Widgets/SModalTaskWindow.cpp
@@ -20,6 +20,7 @@ SModalTaskWindow::~SModalTaskWindow()
 		delete Thread;
 	}
 
+	FPlatformProcess::ReturnSynchEventToPool(AbortEvent);
 	FPlatformProcess::ReturnSynchEventToPool(CloseEvent);
 }
 
