Actually support user filters

diff --git a/Engine/Source/Programs/SlateUGS/Private/UGSCore/Workspace.cpp b/Engine/Source/Programs/SlateUGS/Private/UGSCore/Workspace.cpp
--- a/Engine/Source/Programs/SlateUGS/Private/UGSCore/Workspace.cpp
+++ b/Engine/Source/Programs/SlateUGS/Private/UGSCore/Workspace.cpp
@@ -407,7 +407,7 @@ EWorkspaceUpdateResult FWorkspace::UpdateWorkspaceInternal(FWorkspaceUpdateConte
 
 				if(UserFilter.IsValid())
 				{
-//					SyncRecords.RemoveAll([&UserFilter](const FPerforceFileRecord& SyncRecord){ return SyncRecord.ClientPath.Len() > 0 && !MatchFilter(FPaths::Con.GetFullPath(x.ClientPath), UserFilter));
+					SyncRecords.RemoveAll([&UserFilter](const FPerforceFileRecord& SyncRecord){ return SyncRecord.ClientPath.Len() > 0 && !UserFilter->Matches(SyncRecord.ClientPath); });
 				}
 
 				for(const FPerforceFileRecord& SyncRecord : SyncRecords)
diff --git a/Engine/Source/Programs/SlateUGS/Private/UGSCore/Workspace.h b/Engine/Source/Programs/SlateUGS/Private/UGSCore/Workspace.h
--- a/Engine/Source/Programs/SlateUGS/Private/UGSCore/Workspace.h
+++ b/Engine/Source/Programs/SlateUGS/Private/UGSCore/Workspace.h
@@ -123,7 +123,6 @@ private:
 	virtual uint32 Run() override;
 	void UpdateWorkspace(FWorkspaceUpdateContext& Context);
 	EWorkspaceUpdateResult UpdateWorkspaceInternal(FWorkspaceUpdateContext& Context, FString& OutStatusMessage);
-	bool MatchFilter(const FString& FileName, const FFileFilter& Filter);
 
 	static TArray<FString> GetSyncPaths(const FString& ClientRootPath, const FString& SelectedClientFileName);
 	static TSharedRef<FCustomConfigFile, ESPMode::ThreadSafe> ReadProjectConfigFile(const FString& LocalRootPath, const FString& SelectedLocalFileName, FLineBasedTextWriter& Log);
