Write the log output in UTF8 vs TCHAR

diff --git a/Engine/Source/Programs/SlateUGS/Private/Widgets/SLogWidget.cpp b/Engine/Source/Programs/SlateUGS/Private/Widgets/SLogWidget.cpp
--- a/Engine/Source/Programs/SlateUGS/Private/Widgets/SLogWidget.cpp
+++ b/Engine/Source/Programs/SlateUGS/Private/Widgets/SLogWidget.cpp
@@ -22,7 +22,7 @@ void FLogWidgetTextLayoutMarshaller::SetText(const FString& SourceString, FTextL
 {
 	TextLayout = &TargetTextLayout;
 
-	for(const TSharedRef<FString>& Line : Lines)
+	for (const TSharedRef<FString>& Line : Lines)
 	{
 		TextLayout->AddLine(FSlateTextLayout::FNewLineData(Line, TArray<TSharedRef<IRun>>()));
 	}
@@ -44,10 +44,10 @@ void FLogWidgetTextLayoutMarshaller::AppendLine(const FString& Line)
 	TSharedRef<FString> NewLine = MakeShared<FString>(Line);
 	Lines.Add(NewLine);
 
-	if(TextLayout)
+	if (TextLayout)
 	{
 		// Remove the "default" line that's added for an empty text box.
-		if(Lines.Num() == 1)
+		if (Lines.Num() == 1)
 		{
 			TextLayout->ClearLines();
 		}
@@ -106,16 +106,16 @@ bool SLogWidget::OpenFile(const TCHAR* NewLogFileName)
 	Clear();
 
 	TArray<FString> InitialLines;
-	if(FFileHelper::LoadFileToStringArray(InitialLines, NewLogFileName))
+	if (FFileHelper::LoadFileToStringArray(InitialLines, NewLogFileName))
 	{
-		for(const FString& InitialLine : InitialLines)
+		for (const FString& InitialLine : InitialLines)
 		{
 			AppendLine(InitialLine);
 		}
 	}
 
 	LogWriter = IFileManager::Get().CreateFileWriter(NewLogFileName);
-	if(LogWriter == nullptr)
+	if (LogWriter == nullptr)
 	{
 		return false;
 	}
@@ -126,7 +126,7 @@ bool SLogWidget::OpenFile(const TCHAR* NewLogFileName)
 
 void SLogWidget::CloseFile()
 {
-	if(LogWriter != nullptr)
+	if (LogWriter != nullptr)
 	{
 		delete LogWriter;
 		LogWriter = nullptr;
@@ -135,7 +135,7 @@ void SLogWidget::CloseFile()
 
 void SLogWidget::Clear()
 {
-	if(LogWriter != nullptr)
+	if (LogWriter != nullptr)
 	{
 		delete LogWriter;
 		LogWriter = IFileManager::Get().CreateFileWriter(*LogFileName);
@@ -163,21 +163,24 @@ void SLogWidget::OnScroll(float ScrollOffset)
 EActiveTimerReturnType SLogWidget::OnTimerElapsed(double CurrentTime, float DeltaTime)
 {
 	FScopeLock Lock(&CriticalSection);
-	for(const FString& QueuedLine : QueuedLines)
+	for (const FString& QueuedLine : QueuedLines)
 	{
-		if(LogWriter != nullptr)
+		if (LogWriter != nullptr)
 		{
-			LogWriter->Serialize((void*)*QueuedLine, sizeof(TCHAR) * QueuedLine.Len());
+			FTCHARToUTF8 UTF8String(*QueuedLine, QueuedLine.Len());
+			LogWriter->Serialize((UTF8CHAR*)UTF8String.Get(), UTF8String.Length() * sizeof(UTF8CHAR));
 
-			const TCHAR LineTerminator[] = LINE_TERMINATOR;
+			const ANSICHAR LineTerminator[] = LINE_TERMINATOR_ANSI;
 			LogWriter->Serialize((void*)LineTerminator, sizeof(LineTerminator));
 		}
 		MessagesTextMarshaller->AppendLine(QueuedLine);
 	}
-	if(!bIsUserScrolled)
+
+	if (!bIsUserScrolled)
 	{
 		ScrollToEnd();
 	}
+
 	QueuedLines.Empty();
 	return EActiveTimerReturnType::Continue;
 }
