Rename the user settings to be suffixed with _Slate to avoid stomping exiting settings from c# UGS

diff --git a/Engine/Source/Programs/SlateUGS/Private/UGSTab.cpp b/Engine/Source/Programs/SlateUGS/Private/UGSTab.cpp
--- a/Engine/Source/Programs/SlateUGS/Private/UGSTab.cpp
+++ b/Engine/Source/Programs/SlateUGS/Private/UGSTab.cpp
@@ -265,7 +265,7 @@ void UGSTab::SetupWorkspace()
 	FString DataFolder = FString(FPlatformProcess::UserSettingsDir()) / TEXT("UnrealGameSync");
 	IFileManager::Get().MakeDirectory(*DataFolder);
 
-	UserSettings = MakeShared<FUserSettings>(*(DataFolder / TEXT("UnrealGameSync.ini")));
+	UserSettings = MakeShared<FUserSettings>(*(DataFolder / TEXT("UnrealGameSync_Slate.ini")));
 
 	PerforceClient = DetectSettings->PerforceClient;
 	WorkspaceSettings = UserSettings->FindOrAddWorkspace(*DetectSettings->BranchClientPath);
