Use a callback when the workspace has finished syncing to callback into UGSTab to update WorkspaceSettings

diff --git a/Engine/Source/Programs/SlateUGS/Private/UGSCore/Workspace.cpp b/Engine/Source/Programs/SlateUGS/Private/UGSCore/Workspace.cpp
--- a/Engine/Source/Programs/SlateUGS/Private/UGSCore/Workspace.cpp
+++ b/Engine/Source/Programs/SlateUGS/Private/UGSCore/Workspace.cpp
@@ -342,10 +342,10 @@ uint32 FWorkspace::Run()
 	PendingChangeNumber = CurrentChangeNumber;
 	FPlatformAtomics::InterlockedCompareExchangePointer((void**)&ActiveWorkspace, nullptr, this);
 
-//	if(OnUpdateComplete != null)
-//	{
-//		OnUpdateComplete(Context, Result, StatusMessage);
-//	}
+	if (OnUpdateComplete)
+	{
+		OnUpdateComplete(WorkerThreadContext.ToSharedRef(), Result, StatusMessage);
+	}
 
 	WorkerThreadContext.Reset();
 	return 0;
diff --git a/Engine/Source/Programs/SlateUGS/Private/UGSCore/Workspace.h b/Engine/Source/Programs/SlateUGS/Private/UGSCore/Workspace.h
--- a/Engine/Source/Programs/SlateUGS/Private/UGSCore/Workspace.h
+++ b/Engine/Source/Programs/SlateUGS/Private/UGSCore/Workspace.h
@@ -80,7 +80,7 @@ public:
 	const FString SelectedClientFileName;
 	const FString TelemetryProjectPath;
 
-	TFunction<void(TSharedRef<FWorkspaceUpdateContext, ESPMode::ThreadSafe>, EWorkspaceUpdateResult, FString)> OnUpdateComplete;
+	TFunction<void(TSharedRef<FWorkspaceUpdateContext, ESPMode::ThreadSafe>, EWorkspaceUpdateResult, const FString&)> OnUpdateComplete;
 
 	FWorkspace(TSharedRef<FPerforceConnection> InPerforce, const FString& InLocalRootPath, const FString& InSelectedLocalFileName, const FString& InClientRootPath, const FString& InSelectedClientFileName, int InInitialChangeNumber, int InLastBuiltChangeNumber, const FString& InTelemetryProjectPath, TSharedRef<FLineBasedTextWriter> InLog);
 	~FWorkspace();
diff --git a/Engine/Source/Programs/SlateUGS/Private/UGSTab.cpp b/Engine/Source/Programs/SlateUGS/Private/UGSTab.cpp
--- a/Engine/Source/Programs/SlateUGS/Private/UGSTab.cpp
+++ b/Engine/Source/Programs/SlateUGS/Private/UGSTab.cpp
@@ -200,11 +200,13 @@ bool UGSTab::OnWorkspaceChosen(const FString& Project)
 	if (bIsDataValid)
 	{
 		ProjectFileName = Project;
+
 		SetupWorkspace();
 		GameSyncTabView->SetStreamPathText(FText::FromString(DetectSettings->StreamName));
 		GameSyncTabView->SetProjectPathText(FText::FromString(ProjectFileName));
 		TabWidget->SetContent(GameSyncTabView); // Todo: Set GameSyncTabView data
 		TabWidget->SetLabel(FText::FromString(DetectSettings->StreamName));
+
 		return true;
 	}
 
@@ -244,6 +246,20 @@ bool UGSTab::IsSyncing() const
 	return false;
 }
 
+void UGSTab::OnWorkspaceSyncComplete(TSharedRef<FWorkspaceUpdateContext, ESPMode::ThreadSafe> WorkspaceContext, EWorkspaceUpdateResult SyncResult, const FString& StatusMessage)
+{
+	WorkspaceSettings->CurrentChangeNumber = Workspace->GetCurrentChangeNumber();
+	WorkspaceSettings->LastBuiltChangeNumber = Workspace->GetLastBuiltChangeNumber();
+	WorkspaceSettings->LastSyncResult = SyncResult;
+	WorkspaceSettings->LastSyncResultMessage = StatusMessage;
+
+	WorkspaceSettings->LastSyncTime = WorkspaceContext->StartTime;
+	// TODO check this is valid, may be off
+	WorkspaceSettings->LastSyncDurationSeconds = (FDateTime::UtcNow() - WorkspaceContext->StartTime).GetSeconds();
+
+	UserSettings->Save();
+}
+
 void UGSTab::SetupWorkspace()
 {
 	ProjectFileName = FUtility::GetPathWithCorrectCase(ProjectFileName);
@@ -301,6 +317,10 @@ void UGSTab::SetupWorkspace()
 		TelemetryProjectIdentifier,
 		MakeShared<FLogWidgetTextWriter>(GameSyncTabView->GetSyncLog().ToSharedRef()));
 
+	Workspace->OnUpdateComplete = [this] (TSharedRef<FWorkspaceUpdateContext, ESPMode::ThreadSafe> WorkspaceContext, EWorkspaceUpdateResult SyncResult, const FString& StatusMessage) {
+		OnWorkspaceSyncComplete(WorkspaceContext, SyncResult, StatusMessage);
+	};
+
 	// Todo: Eventually move into the sync operation
 	CombinedSyncFilter = FUserSettings::GetCombinedSyncFilter(
 		Workspace->GetSyncCategories(),
diff --git a/Engine/Source/Programs/SlateUGS/Private/UGSTab.h b/Engine/Source/Programs/SlateUGS/Private/UGSTab.h
--- a/Engine/Source/Programs/SlateUGS/Private/UGSTab.h
+++ b/Engine/Source/Programs/SlateUGS/Private/UGSTab.h
@@ -26,6 +26,11 @@ public:
 	bool IsSyncing() const;
 private:
 
+	void OnWorkspaceSyncComplete(
+		TSharedRef<FWorkspaceUpdateContext, ESPMode::ThreadSafe> WorkspaceContext,
+		EWorkspaceUpdateResult SyncResult,
+		const FString& StatusMessage);
+
 	// Core functions
 	void SetupWorkspace();
 
