Updated sync progress text
Factored error window popup code out into a generic error window widget for later reuse
Converted SWorkspaceWindow into an SWindow (from an SCompoundWidget)

diff --git a/Engine/Source/Programs/SlateUGS/Private/UGSTab.cpp b/Engine/Source/Programs/SlateUGS/Private/UGSTab.cpp
--- a/Engine/Source/Programs/SlateUGS/Private/UGSTab.cpp
+++ b/Engine/Source/Programs/SlateUGS/Private/UGSTab.cpp
@@ -203,8 +203,9 @@ bool UGSTab::OnWorkspaceChosen(const FString& Project)
 
 		SetupWorkspace();
 		GameSyncTabView->SetStreamPathText(FText::FromString(DetectSettings->StreamName));
+		GameSyncTabView->SetChangelistText(FText::FromString(FString::FromInt(Workspace->GetPendingChangeNumber()))); // Todo: also set changelist when completing a sync
 		GameSyncTabView->SetProjectPathText(FText::FromString(ProjectFileName));
-		TabWidget->SetContent(GameSyncTabView); // Todo: Set GameSyncTabView data
+		TabWidget->SetContent(GameSyncTabView);
 		TabWidget->SetLabel(FText::FromString(DetectSettings->StreamName));
 
 		return true;
@@ -238,12 +239,12 @@ void UGSTab::OnSyncLatest()
 
 bool UGSTab::IsSyncing() const
 {
-	if (Workspace.IsValid())
-	{
-		return Workspace->IsBusy();
-	}
+	return Workspace->IsBusy();
+}
 
-	return false;
+FString UGSTab::GetSyncProgress() const
+{
+	return Workspace->GetCurrentProgress().Key;
 }
 
 void UGSTab::OnWorkspaceSyncComplete(TSharedRef<FWorkspaceUpdateContext, ESPMode::ThreadSafe> WorkspaceContext, EWorkspaceUpdateResult SyncResult, const FString& StatusMessage)
diff --git a/Engine/Source/Programs/SlateUGS/Private/UGSTab.h b/Engine/Source/Programs/SlateUGS/Private/UGSTab.h
--- a/Engine/Source/Programs/SlateUGS/Private/UGSTab.h
+++ b/Engine/Source/Programs/SlateUGS/Private/UGSTab.h
@@ -24,6 +24,7 @@ public:
 	bool OnWorkspaceChosen(const FString& Project);
 	void OnSyncLatest();
 	bool IsSyncing() const;
+	FString GetSyncProgress() const;
 private:
 
 	void OnWorkspaceSyncComplete(
diff --git a/Engine/Source/Programs/SlateUGS/Private/Widgets/SEmptyTab.cpp b/Engine/Source/Programs/SlateUGS/Private/Widgets/SEmptyTab.cpp
--- a/Engine/Source/Programs/SlateUGS/Private/Widgets/SEmptyTab.cpp
+++ b/Engine/Source/Programs/SlateUGS/Private/Widgets/SEmptyTab.cpp
@@ -67,15 +67,8 @@ void SEmptyTab::Construct(const FArguments& InArgs)
 
 FReply SEmptyTab::OnOpenProjectClicked()
 {
-	TSharedRef<SWindow> Window = SNew(SWindow)
-	.Title(LOCTEXT("WindowTitle", "Open Project"))
-	.SizingRule(ESizingRule::Autosized)
-	[
-		SNew(SWorkspaceWindow).Tab(Tab)
-	];
-
 	FSlateApplication& SlateApplication = FSlateApplication::Get();
-	SlateApplication.AddModalWindow(Window, Tab->GetTabArgs().GetOwnerWindow(), false);
+	SlateApplication.AddModalWindow(SNew(SWorkspaceWindow).Tab(Tab), Tab->GetTabArgs().GetOwnerWindow(), false);
 
 	return FReply::Handled();
 }
diff --git a/Engine/Source/Programs/SlateUGS/Private/Widgets/SErrorWindow.cpp b/Engine/Source/Programs/SlateUGS/Private/Widgets/SErrorWindow.cpp
new file mode 100644
--- /dev/null
+++ b/Engine/Source/Programs/SlateUGS/Private/Widgets/SErrorWindow.cpp
@@ -0,0 +1,41 @@
+// Copyright Epic Games, Inc. All Rights Reserved.
+
+#include "SErrorWindow.h"
+
+#define LOCTEXT_NAMESPACE "ErrorWindow"
+
+void SErrorWindow::Construct(const FArguments& InArgs)
+{
+	SWindow::Construct(SWindow::FArguments()
+	.Title(LOCTEXT("WindowTitle", "Error"))
+	.SizingRule(ESizingRule::Autosized)
+	.MaxWidth(400)
+	[
+		SNew(SVerticalBox)
+		+SVerticalBox::Slot()
+		.AutoHeight()
+		.HAlign(HAlign_Center)
+		.Padding(25.0f, 25.0f)
+		[
+			SNew(STextBlock)
+			.AutoWrapText(true)
+			.Justification(ETextJustify::Center)
+			.Text(InArgs._ErrorText)
+		]
+		+SVerticalBox::Slot()
+		.AutoHeight()
+		.HAlign(HAlign_Center)
+		[
+			SNew(SButton)
+			.HAlign(HAlign_Center)
+			.Text(LOCTEXT("ErrorWindowOkayButtonText", "Ok"))
+			.OnClicked_Lambda([this]()
+			{
+				SharedThis(this)->RequestDestroyWindow();
+				return FReply::Handled();
+			})
+		]
+	]);
+}
+
+#undef LOCTEXT_NAMESPACE
diff --git a/Engine/Source/Programs/SlateUGS/Private/Widgets/SErrorWindow.h b/Engine/Source/Programs/SlateUGS/Private/Widgets/SErrorWindow.h
new file mode 100644
--- /dev/null
+++ b/Engine/Source/Programs/SlateUGS/Private/Widgets/SErrorWindow.h
@@ -0,0 +1,15 @@
+// Copyright Epic Games, Inc. All Rights Reserved.
+
+#pragma once
+
+#include "Widgets/SWindow.h"
+
+class SErrorWindow : public SWindow
+{
+public:
+	SLATE_BEGIN_ARGS(SErrorWindow) {}
+		SLATE_ARGUMENT(FText, ErrorText)
+	SLATE_END_ARGS()
+
+	void Construct(const FArguments& InArgs);
+};
diff --git a/Engine/Source/Programs/SlateUGS/Private/Widgets/SGameSyncTab.cpp b/Engine/Source/Programs/SlateUGS/Private/Widgets/SGameSyncTab.cpp
--- a/Engine/Source/Programs/SlateUGS/Private/Widgets/SGameSyncTab.cpp
+++ b/Engine/Source/Programs/SlateUGS/Private/Widgets/SGameSyncTab.cpp
@@ -289,6 +289,24 @@ void SGameSyncTab::Construct(const FArguments& InArgs)
 						]
 					]
 					+SVerticalBox::Slot()
+					.Padding(10.0f, 12.5f)
+					[
+						SNew(SHorizontalBox)
+						+SHorizontalBox::Slot()
+						.Padding(5.0f, 0.0f)
+						.HAlign(HAlign_Right)
+						[
+							SNew(STextBlock)
+							.Text(LOCTEXT("ChangelistText", "CHANGELIST"))
+						]
+						+SHorizontalBox::Slot()
+						.HAlign(HAlign_Left)
+						[
+							SAssignNew(ChangelistText, STextBlock)
+							.Text(LOCTEXT("ChangelistTextValue", "No changelist found"))
+						]
+					]
+					+SVerticalBox::Slot()
 					.Padding(10.0f, 25.0f)
 					[
 						SNew(SHorizontalBox)
@@ -315,14 +333,15 @@ void SGameSyncTab::Construct(const FArguments& InArgs)
 					SNew(SHorizontalBox) // Todo: Only display this widget when syncing
 					.Visibility_Lambda([this] { return Tab->IsSyncing() ? EVisibility::Visible : EVisibility::Hidden; })
 					+SHorizontalBox::Slot()
-					.Padding(5.0f, 10.0f)
+					.Padding(5.0f, 25.0f)
 					.AutoWidth()
 					[
-						SNew(STextBlock)
+						SAssignNew(SyncProgressText, STextBlock)
 						.Text(LOCTEXT("SyncProgress", "Syncing Files"))
+						.Text_Lambda([this] { return FText::FromString(Tab->GetSyncProgress()); })
 					]
 					+SHorizontalBox::Slot()
-					.Padding(0.0f, 5.0f, 10.0f, 5.0f)
+					.Padding(0.0f, 12.5f, 12.5f, 5.0f)
 					.AutoWidth()
 					[
 						SNew(SThrobber)
@@ -395,12 +414,14 @@ void SGameSyncTab::SetSyncLogLocation(const FString& LogFileName)
 
 void SGameSyncTab::SetStreamPathText(FText StreamPath)
 {
-	fprintf(stderr, "Setting stream path to: %s\n", TCHAR_TO_ANSI(*StreamPath.ToString()));
 	StreamPathText->SetText(StreamPath);
 }
+void SGameSyncTab::SetChangelistText(FText Changelist)
+{
+	ChangelistText->SetText(Changelist);
+}
 void SGameSyncTab::SetProjectPathText(FText ProjectPath)
 {
-	fprintf(stderr, "Setting project path to: %s\n", TCHAR_TO_ANSI(*ProjectPath.ToString()));
 	ProjectPathText->SetText(ProjectPath);
 }
 
diff --git a/Engine/Source/Programs/SlateUGS/Private/Widgets/SGameSyncTab.h b/Engine/Source/Programs/SlateUGS/Private/Widgets/SGameSyncTab.h
--- a/Engine/Source/Programs/SlateUGS/Private/Widgets/SGameSyncTab.h
+++ b/Engine/Source/Programs/SlateUGS/Private/Widgets/SGameSyncTab.h
@@ -24,7 +24,7 @@ public:
 	void SetSyncLogLocation(const FString& LogFileName);
 
 	void SetStreamPathText(FText StreamPath);
-	// void SetChangelistText(FText Changelist); // Todo: Display the LAST SYNC'd changelist
+	void SetChangelistText(FText Changelist);
 	void SetProjectPathText(FText ProjectPath);
 
 private:
@@ -37,7 +37,9 @@ private:
 	TSharedPtr<SLogWidget> SyncLog;
 
 	TSharedPtr<STextBlock> StreamPathText;
+	TSharedPtr<STextBlock> ChangelistText;
 	TSharedPtr<STextBlock> ProjectPathText;
+	TSharedPtr<STextBlock> SyncProgressText;
 
 	static constexpr float HordeBuildRowHorizontalPadding = 10.0f;
 	static constexpr float HordeBuildRowVerticalPadding = 2.5f;
diff --git a/Engine/Source/Programs/SlateUGS/Private/Widgets/SWorkspaceWindow.cpp b/Engine/Source/Programs/SlateUGS/Private/Widgets/SWorkspaceWindow.cpp
--- a/Engine/Source/Programs/SlateUGS/Private/Widgets/SWorkspaceWindow.cpp
+++ b/Engine/Source/Programs/SlateUGS/Private/Widgets/SWorkspaceWindow.cpp
@@ -8,6 +8,7 @@
 #include "UGSTab.h"
 #include "HordeBuildRowInfo.h"
 #include "SGameSyncTab.h"
+#include "SErrorWindow.h"
 
 #define LOCTEXT_NAMESPACE "UGSWorkspaceWindow"
 
@@ -15,7 +16,10 @@ void SWorkspaceWindow::Construct(const FArguments& InArgs)
 {
 	Tab = InArgs._Tab;
 
-	this->ChildSlot
+	SWindow::Construct(SWindow::FArguments()
+	.Title(LOCTEXT("WindowTitle", "Open Project"))
+	.SizingRule(ESizingRule::FixedSize)
+	.ClientSize(FVector2D(600, 300))
 	[
 		SNew(SBox)
 		.Padding(10.0f, 10.0f)
@@ -200,7 +204,7 @@ void SWorkspaceWindow::Construct(const FArguments& InArgs)
 				]
 			]
 		]
-	];
+	]);
 }
 
 FReply SWorkspaceWindow::OnOkClicked()
@@ -212,40 +216,9 @@ FReply SWorkspaceWindow::OnOkClicked()
 	}
 	else
 	{
-		// Todo: factor out into an error window widget class (check if one already exists)
-		TSharedRef<SWindow> Window = SNew(SWindow)
-		.Title(LOCTEXT("ErrorWindowTitle", "Error Opening Project"))
-		.SizingRule(ESizingRule::Autosized)
-		.MaxWidth(400.0f)
-		[
-			SNew(SVerticalBox)
-			+SVerticalBox::Slot()
-			.AutoHeight()
-			.HAlign(HAlign_Center)
-			.Padding(10.0f, 10.0f)
-			[
-				SNew(STextBlock)
-				.AutoWrapText(true)
-				.Text(LOCTEXT("ErrorText", "Error opening .uproject file, try again")) // Todo: detect the actual reason for error and report it here
-			]
-			+SVerticalBox::Slot()
-			.AutoHeight()
-			.HAlign(HAlign_Center)
-			.Padding(0.0f, 10.0f, 0.0f, 0.0f)
-			[
-				SNew(SButton)
-				.HAlign(HAlign_Center)
-				.Text(LOCTEXT("ErrorWindowOkayButtonText", "Ok"))
-				.OnClicked_Lambda([&Window]() 
-				{
-					Window->RequestDestroyWindow();
-					return FReply::Handled(); 
-				})
-			]
-		];
-
-		FSlateApplication::Get().AddModalWindow(Window, FSlateApplication::Get().GetActiveModalWindow(), false); // Todo: Figure out parent window for modal window
-		// Todo: loading screen widget for detect settings
+		// Todo: report other errors (no workspace associated with the file)
+		FText Error = FText::FromString("Project file does not exist, try again");
+		FSlateApplication::Get().AddModalWindow(SNew(SErrorWindow).ErrorText(Error), SharedThis(this), false);
 	}
 
 	return FReply::Handled();
diff --git a/Engine/Source/Programs/SlateUGS/Private/Widgets/SWorkspaceWindow.h b/Engine/Source/Programs/SlateUGS/Private/Widgets/SWorkspaceWindow.h
--- a/Engine/Source/Programs/SlateUGS/Private/Widgets/SWorkspaceWindow.h
+++ b/Engine/Source/Programs/SlateUGS/Private/Widgets/SWorkspaceWindow.h
@@ -2,15 +2,15 @@
 
 #pragma once
 
-#include "Widgets/SCompoundWidget.h"
+#include "Widgets/SWindow.h"
 
 class UGSTab;
 
-class SWorkspaceWindow : public SCompoundWidget
+class SWorkspaceWindow : public SWindow
 {
 public:
 	SLATE_BEGIN_ARGS(SWorkspaceWindow) {}
-		SLATE_ARGUMENT(UGSTab*, Tab) 
+		SLATE_ARGUMENT(UGSTab*, Tab)
 	SLATE_END_ARGS()
 
 	void Construct(const FArguments& InArgs);
