Default CL text to the workspace settings value. Will be saved on first sync by default will be 0

diff --git a/Engine/Source/Programs/SlateUGS/Private/UGSTab.cpp b/Engine/Source/Programs/SlateUGS/Private/UGSTab.cpp
--- a/Engine/Source/Programs/SlateUGS/Private/UGSTab.cpp
+++ b/Engine/Source/Programs/SlateUGS/Private/UGSTab.cpp
@@ -203,7 +203,7 @@ bool UGSTab::OnWorkspaceChosen(const FString& Project)
 
 		SetupWorkspace();
 		GameSyncTabView->SetStreamPathText(FText::FromString(DetectSettings->StreamName));
-		GameSyncTabView->SetChangelistText(FText::FromString(FString::FromInt(Workspace->GetPendingChangeNumber()))); // Todo: also set changelist when completing a sync
+		GameSyncTabView->SetChangelistText(FText::FromString(FString::FromInt(WorkspaceSettings->CurrentChangeNumber)));
 		GameSyncTabView->SetProjectPathText(FText::FromString(ProjectFileName));
 		TabWidget->SetContent(GameSyncTabView);
 		TabWidget->SetLabel(FText::FromString(DetectSettings->StreamName));
@@ -258,6 +258,8 @@ void UGSTab::OnWorkspaceSyncComplete(TSharedRef<FWorkspaceUpdateContext, ESPMode
 	// TODO check this is valid, may be off
 	WorkspaceSettings->LastSyncDurationSeconds = (FDateTime::UtcNow() - WorkspaceContext->StartTime).GetSeconds();
 
+	GameSyncTabView->SetChangelistText(FText::FromString(FString::FromInt(Workspace->GetCurrentChangeNumber())));
+
 	UserSettings->Save();
 }
 
