Lock up data as we are called back from a thread

diff --git a/Engine/Source/Programs/SlateUGS/Private/UGSTab.cpp b/Engine/Source/Programs/SlateUGS/Private/UGSTab.cpp
--- a/Engine/Source/Programs/SlateUGS/Private/UGSTab.cpp
+++ b/Engine/Source/Programs/SlateUGS/Private/UGSTab.cpp
@@ -247,19 +247,19 @@ FString UGSTab::GetSyncProgress() const
 	return Workspace->GetCurrentProgress().Key;
 }
 
+// This is getting called on a thread, lock our stuff up
 void UGSTab::OnWorkspaceSyncComplete(TSharedRef<FWorkspaceUpdateContext, ESPMode::ThreadSafe> WorkspaceContext, EWorkspaceUpdateResult SyncResult, const FString& StatusMessage)
 {
-	WorkspaceSettings->CurrentChangeNumber = Workspace->GetCurrentChangeNumber();
-	WorkspaceSettings->LastBuiltChangeNumber = Workspace->GetLastBuiltChangeNumber();
-	WorkspaceSettings->LastSyncResult = SyncResult;
-	WorkspaceSettings->LastSyncResultMessage = StatusMessage;
+	FScopeLock Lock(&CriticalSection);
 
-	WorkspaceSettings->LastSyncTime = WorkspaceContext->StartTime;
+	WorkspaceSettings->CurrentChangeNumber     = Workspace->GetCurrentChangeNumber();
+	WorkspaceSettings->LastBuiltChangeNumber   = Workspace->GetLastBuiltChangeNumber();
+	WorkspaceSettings->LastSyncResult          = SyncResult;
+	WorkspaceSettings->LastSyncResultMessage   = StatusMessage;
+	WorkspaceSettings->LastSyncTime            = WorkspaceContext->StartTime;
 	// TODO check this is valid, may be off
 	WorkspaceSettings->LastSyncDurationSeconds = (FDateTime::UtcNow() - WorkspaceContext->StartTime).GetSeconds();
 
-	GameSyncTabView->SetChangelistText(FText::FromString(FString::FromInt(Workspace->GetCurrentChangeNumber())));
-
 	UserSettings->Save();
 }
 
diff --git a/Engine/Source/Programs/SlateUGS/Private/UGSTab.h b/Engine/Source/Programs/SlateUGS/Private/UGSTab.h
--- a/Engine/Source/Programs/SlateUGS/Private/UGSTab.h
+++ b/Engine/Source/Programs/SlateUGS/Private/UGSTab.h
@@ -35,6 +35,8 @@ private:
 	// Core functions
 	void SetupWorkspace();
 
+	FCriticalSection CriticalSection;
+
 	// Slate Data
 	FSpawnTabArgs TabArgs;
 	TSharedRef<SDockTab> TabWidget;
