Add Perforce/Event monitoring code. This will allow us to gather data for the main table used

diff --git a/Engine/Source/Programs/SlateUGS/Private/UGSCore/EventMonitor.cpp b/Engine/Source/Programs/SlateUGS/Private/UGSCore/EventMonitor.cpp
new file mode 100644
--- /dev/null
+++ b/Engine/Source/Programs/SlateUGS/Private/UGSCore/EventMonitor.cpp
@@ -0,0 +1,826 @@
+// Copyright 1998-2017 Epic Games, Inc. All Rights Reserved.
+
+#include "EventMonitor.h"
+#include "HAL/Event.h"
+#include "HAL/PlatformProcess.h"
+#include "Misc/ScopeLock.h"
+
+//// FEventData ////
+
+FEventData::FEventData(int32 InChangeNumber, const FString& InUserName, EEventType InType, const FString& InProject)
+	: FEventData(INT64_MAX, InChangeNumber, InUserName, InType, InProject)
+{
+}
+
+FEventData::FEventData(int64 InId, int32 InChangeNumber, const FString& InUserName, EEventType InType, const FString& InProject)
+	: Id(InId)
+	, ChangeNumber(InChangeNumber)
+	, UserName(InUserName)
+	, Type(InType)
+	, Project(InProject)
+{
+}
+
+//// FCommentData ////
+
+FCommentData::FCommentData(int32 InChangeNumber, const FString& InUserName, const FString& InText, const FString& InProject)
+	: FCommentData(INT64_MAX, InChangeNumber, InUserName, InText, InProject)
+{
+}
+
+FCommentData::FCommentData(int64 InId, int32 InChangeNumber, const FString& InUserName, const FString& InText, const FString& InProject)
+	: Id(InId)
+	, ChangeNumber(InChangeNumber)
+	, UserName(InUserName)
+	, Text(InText)
+	, Project(InProject)
+{
+}
+
+//// FEventSummary ////
+
+FEventSummary::FEventSummary(int InChangeNumber)
+	: ChangeNumber(InChangeNumber)
+	, Verdict(EReviewVerdict::Unknown)
+{
+}
+
+//// FEventMonitor ////
+
+FEventMonitor::FEventMonitor(const FString& InSqlConnectionString, const FString& InProject, const FString& InCurrentUserName, const FString& InLogFileName)
+	: SqlConnectionString(InSqlConnectionString)
+	, Project(InProject)
+	, CurrentUserName(InCurrentUserName)
+	, WorkerThread(nullptr)
+	, RefreshEvent(FPlatformProcess::GetSynchEventFromPool())
+	, LogWriter(*InLogFileName, 128 * 1024)
+	, LastEventId(0)
+	, LastCommentId(0)
+	, LastBuildId(0)
+	, bDisposing(false)
+	, bUpdateActiveInvestigations(true)
+{
+	if(SqlConnectionString.Len() == 0)
+	{
+		LastStatusMessage = TEXT("Database functionality disabled due to empty SqlConnectionString.");
+	}
+	else
+	{
+		LogWriter.Logf(TEXT("Using connection string: %s"), *SqlConnectionString);
+	}
+}
+
+FEventMonitor::~FEventMonitor()
+{
+	bDisposing = true;
+	if(WorkerThread != nullptr)
+	{
+		RefreshEvent->Trigger();
+		WorkerThread->WaitForCompletion();
+		WorkerThread = nullptr;
+	}
+	FPlatformProcess::ReturnSynchEventToPool(RefreshEvent);
+}
+
+void FEventMonitor::Start()
+{
+	check(WorkerThread == nullptr);
+	WorkerThread = FRunnableThread::Create(this, TEXT("Event Monitor"));
+}
+
+void FEventMonitor::FilterChanges(const TArray<int>& ChangeNumbers)
+{
+	// Build a lookup for all the change numbers
+	FilterChangeNumbers.Empty();
+	FilterChangeNumbers.Append(ChangeNumbers);
+
+	// Clear out the list of active users for each review we have
+	UserNameToLastSyncEvent.Empty();
+	for(const TTuple<int, TSharedRef<FEventSummary>>& Pair : ChangeNumberToSummary)
+	{
+		Pair.Value->CurrentUsers.Empty();
+	}
+
+	// Add all the user reviews back in again
+	for(const TTuple<int, TSharedRef<FEventSummary>>& Pair : ChangeNumberToSummary)
+	{
+		const FEventSummary& Summary = *Pair.Value;
+		for(const TSharedRef<FEventData>& SyncEvent : Summary.SyncEvents)
+		{
+			ApplyFilteredUpdate(SyncEvent);
+		}
+	}
+}
+
+FString FEventMonitor::GetLastStatusMessage() const
+{
+	FScopeLock Lock(&CriticalSection);
+	return LastStatusMessage;
+}
+
+void FEventMonitor::ApplyUpdates()
+{
+	// Create local copies of all the incoming data
+	CriticalSection.Lock();
+
+	TArray<TSharedRef<FEventData>> LocalIncomingEvents; 
+	Exchange(LocalIncomingEvents, IncomingEvents);
+
+	TArray<TSharedRef<FBuildData>> LocalIncomingBuilds;
+	Exchange(LocalIncomingBuilds, IncomingBuilds);
+
+	TArray<TSharedRef<FCommentData>> LocalIncomingComments;
+	Exchange(LocalIncomingComments, IncomingComments);
+
+	CriticalSection.Unlock();
+
+	// Apply all the updates
+	for(const TSharedRef<FEventData>& Event : LocalIncomingEvents)
+	{
+		ApplyEventUpdate(Event);
+	}
+
+	for(const TSharedRef<FBuildData>& Build : LocalIncomingBuilds)
+	{
+		ApplyBuildUpdate(Build);
+	}
+
+	for(const TSharedRef<FCommentData>& Comment : LocalIncomingComments)
+	{
+		ApplyCommentUpdate(Comment);
+	}
+}
+
+void FEventMonitor::PostEvent(int ChangeNumber, EEventType Type)
+{
+	if(SqlConnectionString.Len() > 0)
+	{
+		TSharedRef<FEventData> Event(new FEventData(ChangeNumber, CurrentUserName, Type, Project));
+
+		CriticalSection.Lock();
+		OutgoingEvents.Add(Event);
+		CriticalSection.Unlock();
+
+		ApplyEventUpdate(Event);
+
+		RefreshEvent->Trigger();
+	}
+}
+
+void FEventMonitor::PostComment(int ChangeNumber, const FString& Text)
+{
+	if(SqlConnectionString.Len() > 0)
+	{
+		TSharedRef<FCommentData> Comment(new FCommentData(ChangeNumber, CurrentUserName, Text, Project));
+
+		CriticalSection.Lock();
+		OutgoingComments.Add(Comment);
+		CriticalSection.Unlock();
+
+		ApplyCommentUpdate(Comment);
+
+		RefreshEvent->Trigger();
+	}
+}
+
+bool FEventMonitor::TryGetCommentByCurrentUser(int ChangeNumber, FString& OutCommentText) const
+{
+	TSharedPtr<FEventSummary> Summary;
+	if(TryGetSummaryForChange(ChangeNumber, Summary))
+	{
+		for(const TSharedRef<FCommentData>& Comment : Summary->Comments)
+		{
+			if(Comment->UserName == CurrentUserName && Comment->Text.Len() > 0)
+			{
+				OutCommentText = Comment->Text;
+				return true;
+			}
+		}
+	}
+	return false;
+}
+
+bool FEventMonitor::TryGetReviewByCurrentUser(int ChangeNumber, TSharedPtr<FEventData>& OutReview) const
+{
+	TSharedPtr<FEventSummary> Summary;
+	if(TryGetSummaryForChange(ChangeNumber, Summary))
+	{
+		for(const TSharedRef<FEventData>& Event : Summary->Reviews)
+		{
+			if(Event->UserName == CurrentUserName)
+			{
+				OutReview = Event;
+				return true;
+			}
+		}
+	}
+	return false;
+}
+
+bool FEventMonitor::TryGetSummaryForChange(int ChangeNumber, TSharedPtr<FEventSummary>& OutSummary) const
+{
+	const TSharedRef<FEventSummary>* Summary = ChangeNumberToSummary.Find(ChangeNumber);
+	if(Summary != nullptr)
+	{
+		OutSummary = *Summary;
+		return true;
+	}
+	return false;
+}
+
+bool FEventMonitor::IsReview(EEventType Type)
+{
+	return IsPositiveReview(Type) || IsNegativeReview(Type) || Type == EEventType::Unknown;
+}
+
+bool FEventMonitor::IsPositiveReview(EEventType Type)
+{
+	return Type == EEventType::Good || Type == EEventType::Compiles;
+}
+
+bool FEventMonitor::IsNegativeReview(EEventType Type)
+{
+	return Type == EEventType::DoesNotCompile || Type == EEventType::Bad;
+}
+
+bool FEventMonitor::WasSyncedByCurrentUser(int ChangeNumber) const
+{
+	TSharedPtr<FEventSummary> Summary;
+	if(TryGetSummaryForChange(ChangeNumber, Summary))
+	{
+		for(const TSharedRef<FEventData>& SyncEvent : Summary->SyncEvents)
+		{
+			if(SyncEvent->Type == EEventType::Syncing && SyncEvent->UserName == CurrentUserName)
+			{
+				return true;
+			}
+		}
+	}
+	return false;
+}
+
+void FEventMonitor::StartInvestigating(int ChangeNumber)
+{
+	PostEvent(ChangeNumber, EEventType::Investigating);
+}
+
+void FEventMonitor::FinishInvestigating(int ChangeNumber)
+{
+	PostEvent(ChangeNumber, EEventType::Resolved);
+}
+
+bool FEventMonitor::IsUnderInvestigation(int ChangeNumber)
+{
+	UpdateActiveInvestigations();
+
+	for(const TSharedRef<FEventData>& Investigation : ActiveInvestigations)
+	{
+		if(Investigation->ChangeNumber <= ChangeNumber)
+		{
+			return true;
+		}
+	}
+
+	return false;
+}
+
+bool FEventMonitor::IsUnderInvestigationByCurrentUser(int ChangeNumber)
+{
+	UpdateActiveInvestigations();
+
+	for(const TSharedRef<FEventData>& Investigation : ActiveInvestigations)
+	{
+		if(Investigation->ChangeNumber <= ChangeNumber && Investigation->UserName == CurrentUserName)
+		{
+			return true;
+		}
+	}
+
+	return false;
+}
+
+TArray<FString> FEventMonitor::GetInvestigatingUsers(int ChangeNumber)
+{
+	UpdateActiveInvestigations();
+
+	TArray<FString> UserNames;
+	for(const TSharedRef<FEventData>& Investigation : ActiveInvestigations)
+	{
+		if(Investigation->ChangeNumber <= ChangeNumber)
+		{
+			UserNames.Add(Investigation->UserName);
+		}
+	}
+	return UserNames;
+}
+
+int FEventMonitor::GetInvestigationStartChangeNumber(int LastChangeNumber)
+{
+	UpdateActiveInvestigations();
+
+	int StartChangeNumber = -1;
+	for(const TSharedRef<FEventData>& ActiveInvestigation : ActiveInvestigations)
+	{
+		if(ActiveInvestigation->UserName == CurrentUserName)
+		{
+			if(ActiveInvestigation->ChangeNumber <= LastChangeNumber && (StartChangeNumber == -1 || ActiveInvestigation->ChangeNumber < StartChangeNumber))
+			{
+				StartChangeNumber = ActiveInvestigation->ChangeNumber;
+			}
+		}
+	}
+	return StartChangeNumber;
+}
+
+TSharedRef<FEventSummary> FEventMonitor::FindOrAddSummary(int ChangeNumber)
+{
+	TSharedRef<FEventSummary>* SummaryPtr = ChangeNumberToSummary.Find(ChangeNumber);
+	if(SummaryPtr == nullptr)
+	{
+		SummaryPtr = &(ChangeNumberToSummary.Add(ChangeNumber, TSharedRef<FEventSummary>(new FEventSummary(ChangeNumber))));
+	}
+	return *SummaryPtr;
+}
+
+void FEventMonitor::ApplyEventUpdate(const TSharedRef<FEventData>& Event)
+{
+	TSharedRef<FEventSummary> Summary = FindOrAddSummary(Event->ChangeNumber);
+	if(Event->Type == EEventType::Starred || Event->Type == EEventType::Unstarred)
+	{
+		// If it's a star or un-star review, process that separately
+		if(!Summary->LastStarReview.IsValid() || Event->Id > Summary->LastStarReview->Id)
+		{
+			Summary->LastStarReview = Event;
+		}
+	}
+	else if(Event->Type == EEventType::Investigating || Event->Type == EEventType::Resolved)
+	{
+		// Insert it sorted in the investigation list
+		int InsertIdx = 0;
+		while(InsertIdx < InvestigationEvents.Num() && InvestigationEvents[InsertIdx]->Id < Event->Id)
+		{
+			InsertIdx++;
+		}
+		if(InsertIdx == InvestigationEvents.Num() || InvestigationEvents[InsertIdx]->Id != Event->Id)
+		{
+			InvestigationEvents.Insert(Event, InsertIdx);
+		}
+		bUpdateActiveInvestigations = true;
+	}
+	else if(Event->Type == EEventType::Syncing)
+	{
+		Summary->SyncEvents.RemoveAll([&Event](const TSharedRef<FEventData>& OtherEvent){ return Event->UserName == OtherEvent->UserName; });
+		Summary->SyncEvents.Add(Event);
+		ApplyFilteredUpdate(Event);
+	}
+	else if(IsReview(Event->Type))
+	{
+		// Try to find an existing review by this user. If we already have a newer review, ignore this one. Otherwise remove it.
+		for(int Idx = 0; Idx < Summary->Reviews.Num(); Idx++)
+		{
+			const TSharedRef<FEventData>& ExistingReview = Summary->Reviews[Idx];
+			if(ExistingReview->UserName == Event->UserName)
+			{
+				if(ExistingReview->Id > Event->Id)
+				{
+					return;
+				}
+				Summary->Reviews.RemoveAt(Idx);
+				break;
+			}
+		}
+
+		// Add the new review, and find the new verdict for this change
+		Summary->Reviews.Add(Event);
+		Summary->Verdict = GetVerdict(Summary->Reviews, Summary->Builds);
+	}
+	else
+	{
+		// Unknown type
+	}
+}
+
+void FEventMonitor::ApplyBuildUpdate(const TSharedRef<FBuildData>& Build)
+{
+	TSharedRef<FEventSummary> Summary = FindOrAddSummary(Build->ChangeNumber);
+
+	for(int Idx = 0; Idx < Summary->Builds.Num(); Idx++)
+	{
+		const TSharedRef<FBuildData>& ExistingBuild = Summary->Builds[Idx];
+		if(ExistingBuild->ChangeNumber == Build->ChangeNumber && ExistingBuild->BuildType == Build->BuildType)
+		{
+			if(ExistingBuild->Id > Build->Id)
+			{
+				return;
+			}
+			Summary->Builds.RemoveAt(Idx);
+			break;
+		}
+	}
+
+	Summary->Builds.Add(Build);
+	Summary->Verdict = GetVerdict(Summary->Reviews, Summary->Builds);
+}
+
+void FEventMonitor::ApplyCommentUpdate(const TSharedRef<FCommentData>& Comment)
+{
+//	TSharedRef<FEventSummary> Summary = FindOrAddSummary(Comment->ChangeNumber);
+//	if(String.Compare(Comment.UserName, CurrentUserName, true) == 0 && Summary.Comments.Count > 0 && Summary.Comments.Last().Id == long.MaxValue)
+//	{
+//		// This comment was added by PostComment(), to mask the latency of a round trip to the server. Remove it now we have the sorted comment.
+//		Summary.Comments.RemoveAt(Summary.Comments.Count - 1);
+//	}
+//	AddPerUserItem(Summary.Comments, Comment, x => x.Id, x => x.UserName);
+}
+/*
+static bool AddPerUserItem<T>(List<T> Items, T NewItem, Func<T, long> IdSelector, Func<T, string> UserSelector)
+{
+	int InsertIdx = Items.Count;
+
+	for(; InsertIdx > 0 && IdSelector(Items[InsertIdx - 1]) >= IdSelector(NewItem); InsertIdx--)
+	{
+		if(String.Compare(UserSelector(Items[InsertIdx - 1]), UserSelector(NewItem), true) == 0)
+		{
+			return false;
+		}
+	}
+
+	Items.Insert(InsertIdx, NewItem);
+
+	for(; InsertIdx > 0; InsertIdx--)
+	{
+		if(String.Compare(UserSelector(Items[InsertIdx - 1]), UserSelector(NewItem), true) == 0)
+		{
+			Items.RemoveAt(InsertIdx - 1);
+		}
+	}
+
+	return true;
+}
+*/
+EReviewVerdict FEventMonitor::GetVerdict(const TArray<TSharedRef<FEventData>>& Events, const TArray<TSharedRef<FBuildData>>& Builds)
+{
+	int NumPositiveReviews = 0;
+	int NumNegativeReviews = 0;
+
+	int NumCompiles = 0;
+	int NumFailedCompiles = 0;
+
+	for(const TSharedRef<FEventData>& Event : Events)
+	{
+		if(Event->Type == EEventType::Good)
+		{
+			NumPositiveReviews++;
+		}
+		else if(Event->Type == EEventType::Bad)
+		{
+			NumNegativeReviews++;
+		}
+		else if(Event->Type == EEventType::Compiles)
+		{
+			NumCompiles++;
+		}
+		else if(Event->Type == EEventType::DoesNotCompile)
+		{
+			NumFailedCompiles++;
+		}
+	}
+
+	if(NumPositiveReviews > 0 || NumNegativeReviews > 0)
+	{
+		return GetVerdict(NumPositiveReviews, NumNegativeReviews);
+	}
+
+	if(NumCompiles > 0 || NumFailedCompiles > 0)
+	{
+		return GetVerdict(NumCompiles, NumFailedCompiles);
+	}
+
+	int NumBuilds = 0;
+	int NumFailedBuilds = 0;
+
+	for(const TSharedRef<FBuildData>& Build : Builds)
+	{
+		if(Build->BuildType == TEXT("Editor"))
+		{
+			if(Build->IsSuccess())
+			{
+				NumBuilds++;
+			}
+			if(Build->IsFailure())
+			{
+				NumFailedBuilds++;
+			}
+		}
+	}
+
+	if(NumBuilds > 0 || NumFailedBuilds > 0)
+	{
+		return GetVerdict(NumBuilds, NumFailedBuilds);
+	}
+
+	return EReviewVerdict::Unknown;
+}
+
+EReviewVerdict FEventMonitor::GetVerdict(int NumPositive, int NumNegative)
+{
+	if(NumPositive > (int)(NumNegative * 1.5))
+	{
+		return EReviewVerdict::Good;
+	}
+	else if(NumPositive >= NumNegative)
+	{
+		return EReviewVerdict::Mixed;
+	}
+	else
+	{
+		return EReviewVerdict::Bad;
+	}
+}
+
+void FEventMonitor::ApplyFilteredUpdate(const TSharedRef<FEventData>& Event)
+{
+	if(Event->Type == EEventType::Syncing && FilterChangeNumbers.Contains(Event->ChangeNumber))
+	{
+		// Update the active users list for this change
+		TSharedRef<FEventData>* LastSync = UserNameToLastSyncEvent.Find(Event->UserName);
+		if(LastSync == nullptr)
+		{
+			FindOrAddSummary(Event->ChangeNumber)->CurrentUsers.Add(Event->UserName);
+			UserNameToLastSyncEvent[Event->UserName] = Event;
+		}
+		else if(Event->Id > (*LastSync)->Id)
+		{
+			ChangeNumberToSummary[(*LastSync)->ChangeNumber]->CurrentUsers.Remove(Event->UserName);
+			FindOrAddSummary(Event->ChangeNumber)->CurrentUsers.Add(Event->UserName);
+			UserNameToLastSyncEvent[Event->UserName] = Event;
+		}
+	}
+}
+
+uint32 FEventMonitor::Run()
+{
+	TArray<TSharedRef<FEventData>> LocalOutgoingEvents;
+	TArray<TSharedRef<FCommentData>> LocalOutgoingComments;
+	while(!bDisposing)
+	{
+		// If there's no connection string, just empty out the queue
+		if(SqlConnectionString.Len() > 0)
+		{
+			// Append all the queued events to the local queue on this thread
+			CriticalSection.Lock();
+
+			LocalOutgoingEvents += OutgoingEvents;
+			OutgoingEvents.Empty();
+
+			LocalOutgoingComments += OutgoingComments;
+			OutgoingComments.Empty();
+
+			CriticalSection.Unlock();
+
+			// Post all the reviews to the database. We don't send them out of order, so keep the review outside the queue until the next update if it fails
+			while(LocalOutgoingEvents.Num() > 0 && SendEventToBackend(LocalOutgoingEvents[0]))
+			{
+				LocalOutgoingEvents.RemoveAt(0);
+			}
+
+			// Post all the comments to the database.
+			while(LocalOutgoingComments.Num() > 0 && SendCommentToBackend(LocalOutgoingComments[0]))
+			{
+				LocalOutgoingComments.RemoveAt(0);
+			}
+
+			// Read all the new reviews
+			ReadEventsFromBackend();
+
+			// Send a notification that we're ready to update
+//			if((IncomingEvents.Num() > 0 || IncomingBuilds.Num() > 0 || IncomingComments.Num() > 0) && OnUpdatesReady != null)
+//			{
+//				OnUpdatesReady();
+//			}
+		}
+
+		// Wait for something else to do
+		RefreshEvent->Wait(FTimespan::FromSeconds(30.0));
+	}
+	return 0;
+}
+
+bool FEventMonitor::SendEventToBackend(const TSharedRef<FEventData>& Event)
+{
+//	try
+//	{
+//		Stopwatch Timer = Stopwatch.StartNew();
+//		LogWriter.WriteLine("Posting event... ({0}, {1}, {2})", Event.Change, Event.UserName, Event.Type);
+//		using(SqlConnection Connection = new SqlConnection(SqlConnectionString))
+//		{
+//			Connection.Open();
+//			using (SqlCommand Command = new SqlCommand("INSERT INTO dbo.UserVotes (Changelist, UserName, Verdict, Project) VALUES (@Changelist, @UserName, @Verdict, @Project)", Connection))
+//			{
+//				Command.Parameters.AddWithValue("@Changelist", Event.Change);
+//				Command.Parameters.AddWithValue("@UserName", Event.UserName.ToString());
+//				Command.Parameters.AddWithValue("@Verdict", Event.Type.ToString());
+//				Command.Parameters.AddWithValue("@Project", Event.Project);
+//				Command.ExecuteNonQuery();
+//			}
+//		}
+//		LogWriter.WriteLine("Done in {0}ms.", Timer.ElapsedMilliseconds);
+//		return true;
+//	}
+//	catch(Exception Ex)
+//	{
+//		LogWriter.WriteException(Ex, "Failed with exception.");
+//		return false;
+//	}
+	return true;
+}
+
+bool FEventMonitor::SendCommentToBackend(const TSharedRef<FCommentData>& Comment)
+{
+//	try
+//	{
+//		Stopwatch Timer = Stopwatch.StartNew();
+//		LogWriter.WriteLine("Posting comment... ({0}, {1}, {2}, {3})", Comment.ChangeNumber, Comment.UserName, Comment.Text, Comment.Project);
+//		using(SqlConnection Connection = new SqlConnection(SqlConnectionString))
+//		{
+//			Connection.Open();
+//			using (SqlCommand Command = new SqlCommand("INSERT INTO dbo.Comments (ChangeNumber, UserName, Text, Project) VALUES (@ChangeNumber, @UserName, @Text, @Project)", Connection))
+//			{
+//				Command.Parameters.AddWithValue("@ChangeNumber", Comment.ChangeNumber);
+//				Command.Parameters.AddWithValue("@UserName", Comment.UserName);
+//				Command.Parameters.AddWithValue("@Text", Comment.Text);
+//				Command.Parameters.AddWithValue("@Project", Comment.Project);
+//				Command.ExecuteNonQuery();
+//			}
+//		}
+//		LogWriter.WriteLine("Done in {0}ms.", Timer.ElapsedMilliseconds);
+//		return true;
+//	}
+//	catch(Exception Ex)
+//	{
+//		LogWriter.WriteException(Ex, "Failed with exception.");
+//		return false;
+//	}
+	return true;
+}
+
+bool FEventMonitor::ReadEventsFromBackend()
+{
+	/*
+	try
+	{
+		Stopwatch Timer = Stopwatch.StartNew();
+		LogWriter.WriteLine();
+		LogWriter.WriteLine("Polling for reviews at {0}...", DateTime.Now.ToString());
+
+		if(LastEventId == 0)
+		{
+			using(SqlConnection Connection = new SqlConnection(SqlConnectionString))
+			{
+				Connection.Open();
+				using (SqlCommand Command = new SqlCommand("SELECT MAX(ID) FROM dbo.UserVotes", Connection))
+				{
+					using (SqlDataReader Reader = Command.ExecuteReader())
+					{
+						while (Reader.Read())
+						{
+							LastEventId = Reader.GetInt64(0);
+							LastEventId = Math.Max(LastEventId - 5000, 0);
+							break;
+						}
+					}
+				}
+			}
+		}
+
+		using(SqlConnection Connection = new SqlConnection(SqlConnectionString))
+		{
+			Connection.Open();
+			using (SqlCommand Command = new SqlCommand("SELECT Id, Changelist, UserName, Verdict, Project FROM dbo.UserVotes WHERE Id > @param1 ORDER BY Id", Connection))
+			{
+				Command.Parameters.AddWithValue("@param1", LastEventId);
+				using (SqlDataReader Reader = Command.ExecuteReader())
+				{
+					while (Reader.Read())
+					{
+						EventData Review = new EventData();
+						Review.Id = Reader.GetInt64(0);
+						Review.Change = Reader.GetInt32(1);
+						Review.UserName = Reader.GetString(2);
+						Review.Project = Reader.IsDBNull(4)? null : Reader.GetString(4);
+						if(Enum.TryParse(Reader.GetString(3), out Review.Type))
+						{
+							if(Review.Project == null || String.Compare(Review.Project, Project, true) == 0)
+							{
+								IncomingEvents.Enqueue(Review);
+							}
+							LastEventId = Math.Max(LastEventId, Review.Id);
+						}
+					}
+				}
+			}
+			using(SqlCommand Command = new SqlCommand("SELECT Id, ChangeNumber, UserName, Text, Project FROM dbo.Comments WHERE Id > @param1 ORDER BY Id", Connection))
+			{
+				Command.Parameters.AddWithValue("@param1", LastCommentId);
+				using (SqlDataReader Reader = Command.ExecuteReader())
+				{
+					while (Reader.Read())
+					{
+						CommentData Comment = new CommentData();
+						Comment.Id = Reader.GetInt32(0);
+						Comment.ChangeNumber = Reader.GetInt32(1);
+						Comment.UserName = Reader.GetString(2);
+						Comment.Text = Reader.GetString(3);
+						Comment.Project = Reader.GetString(4);
+						if(Comment.Project == null || String.Compare(Comment.Project, Project, true) == 0)
+						{
+							IncomingComments.Enqueue(Comment);
+						}
+						LastCommentId = Math.Max(LastCommentId, Comment.Id);
+					}
+				}
+			}
+			using(SqlCommand Command = new SqlCommand("SELECT Id, ChangeNumber, BuildType, Result, Url, Project FROM dbo.CIS WHERE Id > @param1 ORDER BY Id", Connection))
+			{
+				Command.Parameters.AddWithValue("@param1", LastBuildId);
+				using (SqlDataReader Reader = Command.ExecuteReader())
+				{
+					while (Reader.Read())
+					{
+						BuildData Build = new BuildData();
+						Build.Id = Reader.GetInt32(0);
+						Build.ChangeNumber = Reader.GetInt32(1);
+						Build.BuildType = Reader.GetString(2).TrimRight();
+						if(Enum.TryParse(Reader.GetString(3).TrimRight(), true, out Build.Result))
+						{
+							Build.Url = Reader.GetString(4);
+							Build.Project = Reader.IsDBNull(5)? null : Reader.GetString(5);
+							if(Build.Project == null || String.Compare(Build.Project, Project, true) == 0 || MatchesWildcard(Build.Project, Project))
+							{
+								IncomingBuilds.Enqueue(Build);
+							}
+						}
+						LastBuildId = Math.Max(LastBuildId, Build.Id);
+					}
+				}
+			}
+		}
+		LastStatusMessage = String.Format("Last update took {0}ms", Timer.ElapsedMilliseconds);
+		LogWriter.WriteLine("Done in {0}ms.", Timer.ElapsedMilliseconds);
+		return true;
+	}
+	catch(Exception Ex)
+	{
+		LogWriter.WriteException(Ex, "Failed with exception.");
+		LastStatusMessage = String.Format("Last update failed: ({0})", Ex.ToString());
+		return false;
+	}
+	*/
+	return true;
+}
+
+bool FEventMonitor::MatchesWildcard(const FString& Wildcard, const FString& Project)
+{
+	return Wildcard.EndsWith("...") && Project.StartsWith(Wildcard.Mid(0, Wildcard.Len() - 3));
+}
+
+void FEventMonitor::UpdateActiveInvestigations()
+{
+	if(bUpdateActiveInvestigations)
+	{
+		// Insert investigation events into the active list, sorted by change number. Remove 
+		ActiveInvestigations.Empty();
+		for(const TSharedRef<FEventData>& InvestigationEvent : InvestigationEvents)
+		{
+			if(FilterChangeNumbers.Contains(InvestigationEvent->ChangeNumber))
+			{
+				if(InvestigationEvent->Type == EEventType::Investigating)
+				{
+					int InsertIdx = 0;
+					while(InsertIdx < ActiveInvestigations.Num() && ActiveInvestigations[InsertIdx]->ChangeNumber > InvestigationEvent->ChangeNumber)
+					{
+						InsertIdx++;
+					}
+					ActiveInvestigations.Insert(InvestigationEvent, InsertIdx);
+				}
+				else
+				{
+					ActiveInvestigations.RemoveAll([&InvestigationEvent](const TSharedRef<FEventData>& Event){ return Event->UserName == InvestigationEvent->UserName && Event->ChangeNumber <= InvestigationEvent->ChangeNumber; });
+				}
+			}
+		}
+
+		// Remove any duplicate users
+		for(int Idx = 0; Idx < ActiveInvestigations.Num(); Idx++)
+		{
+			for(int OtherIdx = 0; OtherIdx < Idx; OtherIdx++)
+			{
+				if(ActiveInvestigations[Idx]->UserName == ActiveInvestigations[OtherIdx]->UserName)
+				{
+					ActiveInvestigations.RemoveAt(Idx--);
+					break;
+				}
+			}
+		}
+	}
+}
diff --git a/Engine/Source/Programs/SlateUGS/Private/UGSCore/EventMonitor.h b/Engine/Source/Programs/SlateUGS/Private/UGSCore/EventMonitor.h
new file mode 100644
--- /dev/null
+++ b/Engine/Source/Programs/SlateUGS/Private/UGSCore/EventMonitor.h
@@ -0,0 +1,181 @@
+// Copyright 1998-2017 Epic Games, Inc. All Rights Reserved.
+
+#pragma once
+
+#include "CoreMinimal.h"
+#include "OutputAdapters.h"
+#include "HAL/Runnable.h"
+#include "HAL/RunnableThread.h"
+
+enum class EEventType
+{
+	Syncing,
+
+	// Reviews
+	Compiles,
+	DoesNotCompile,
+	Good,
+	Bad,
+	Unknown,
+		
+	// Starred builds
+	Starred,
+	Unstarred,
+
+	// Investigating events
+	Investigating,
+	Resolved,
+};
+
+struct FEventData
+{
+	int64 Id;
+	int32 ChangeNumber;
+	FString UserName;
+	EEventType Type;
+	FString Project;
+
+	FEventData(int32 InChangeNumber, const FString& InUserName, EEventType InType, const FString& InProject);
+	FEventData(int64 InId, int32 InChangeNumber, const FString& InUserName, EEventType InType, const FString& InProject);
+};
+
+struct FCommentData
+{
+	int64 Id;
+	int32 ChangeNumber;
+	FString UserName;
+	FString Text;
+	FString Project;
+
+	FCommentData(int32 InChangeNumber, const FString& InUserName, const FString& InText, const FString& InProject);
+	FCommentData(int64 InId, int32 InChangeNumber, const FString& InUserName, const FString& InText, const FString& InProject);
+};
+
+enum class EBuildDataResult
+{
+	Starting,
+	Failure,
+	Warning,
+	Success,
+};
+
+struct FBuildData
+{
+	long Id;
+	int ChangeNumber;
+	FString BuildType;
+	EBuildDataResult Result;
+	FString Url;
+	FString Project;
+
+	bool IsSuccess() const
+	{
+		return Result == EBuildDataResult::Success || Result == EBuildDataResult::Warning;
+	}
+
+	bool IsFailure() const
+	{
+		return Result == EBuildDataResult::Failure;
+	}
+};
+
+enum class EReviewVerdict
+{
+	Unknown,
+	Good,
+	Bad,
+	Mixed,
+};
+
+struct FEventSummary
+{
+	int ChangeNumber;
+	EReviewVerdict Verdict;
+	TArray<TSharedRef<FEventData>> SyncEvents;
+	TArray<TSharedRef<FEventData>> Reviews;
+	TArray<FString> CurrentUsers;
+	TSharedPtr<FEventData> LastStarReview;
+	TArray<TSharedRef<FBuildData>> Builds;
+	TArray<TSharedRef<FCommentData>> Comments;
+
+	FEventSummary(int InChangeNumber);
+};
+
+class FEventMonitor : FRunnable
+{
+public:
+	TFunction<void()> OnUpdatesReady;
+
+	FEventMonitor(const FString& InSqlConnectionString, const FString& InProject, const FString& InCurrentUserName, const FString& InLogFileName);
+	~FEventMonitor();
+
+	void Start();
+	void FilterChanges(const TArray<int>& ChangeNumbers);
+
+	FString GetLastStatusMessage() const;
+
+	void ApplyUpdates();
+	void PostEvent(int ChangeNumber, EEventType Type);
+	void PostComment(int ChangeNumber, const FString& Text);
+
+	bool TryGetCommentByCurrentUser(int ChangeNumber, FString& OutCommentText) const;
+	bool TryGetReviewByCurrentUser(int ChangeNumber, TSharedPtr<FEventData>& OutReview) const;
+	bool TryGetSummaryForChange(int ChangeNumber, TSharedPtr<FEventSummary>& OutSummary) const;
+
+	static bool IsReview(EEventType Type);
+	static bool IsPositiveReview(EEventType Type);
+	static bool IsNegativeReview(EEventType Type);
+
+	bool WasSyncedByCurrentUser(int ChangeNumber) const;
+
+	void StartInvestigating(int ChangeNumber);
+	void FinishInvestigating(int ChangeNumber);
+	bool IsUnderInvestigation(int ChangeNumber);
+	bool IsUnderInvestigationByCurrentUser(int ChangeNumber);
+	TArray<FString> GetInvestigatingUsers(int ChangeNumber);
+	int GetInvestigationStartChangeNumber(int LastChangeNumber);
+
+private:
+	const FString SqlConnectionString;
+	const FString Project;
+	const FString CurrentUserName;
+	FRunnableThread* WorkerThread;
+	mutable FCriticalSection CriticalSection;
+	FEvent* RefreshEvent;
+	TArray<TSharedRef<FEventData>> OutgoingEvents;
+	TArray<TSharedRef<FEventData>> IncomingEvents;
+	TArray<TSharedRef<FCommentData>> OutgoingComments;
+	TArray<TSharedRef<FCommentData>> IncomingComments;
+	TArray<TSharedRef<FBuildData>> IncomingBuilds;
+	TMap<int, TSharedRef<FEventSummary>> ChangeNumberToSummary;
+	TMap<FString, TSharedRef<FEventData>> UserNameToLastSyncEvent;
+	FBoundedLogWriter LogWriter;
+	long LastEventId;
+	long LastCommentId;
+	long LastBuildId;
+	bool bDisposing;
+	FString LastStatusMessage;
+	TSet<int> FilterChangeNumbers;
+	TArray<TSharedRef<FEventData>> InvestigationEvents;
+	bool bUpdateActiveInvestigations;
+	TArray<TSharedRef<FEventData>> ActiveInvestigations;
+
+	TSharedRef<FEventSummary> FindOrAddSummary(int ChangeNumber);
+	void ApplyEventUpdate(const TSharedRef<FEventData>& Event);
+	void ApplyBuildUpdate(const TSharedRef<FBuildData>& Build);
+	void ApplyCommentUpdate(const TSharedRef<FCommentData>& Comment);
+
+	static EReviewVerdict GetVerdict(const TArray<TSharedRef<FEventData>>& Events, const TArray<TSharedRef<FBuildData>>& Builds);
+	static EReviewVerdict GetVerdict(int NumPositive, int NumNegative);
+
+	void ApplyFilteredUpdate(const TSharedRef<FEventData>& Event);
+	uint32 Run();
+	bool SendEventToBackend(const TSharedRef<FEventData>& Event);
+	bool SendCommentToBackend(const TSharedRef<FCommentData>& Comment);
+	bool ReadEventsFromBackend();
+
+	static bool MatchesWildcard(const FString& Wildcard, const FString& Project);
+
+	void UpdateActiveInvestigations();
+};
+
diff --git a/Engine/Source/Programs/SlateUGS/Private/UGSCore/PerforceMonitor.cpp b/Engine/Source/Programs/SlateUGS/Private/UGSCore/PerforceMonitor.cpp
new file mode 100644
--- /dev/null
+++ b/Engine/Source/Programs/SlateUGS/Private/UGSCore/PerforceMonitor.cpp
@@ -0,0 +1,509 @@
+// Copyright 1998-2017 Epic Games, Inc. All Rights Reserved.
+
+#include "PerforceMonitor.h"
+#include "HAL/PlatformProcess.h"
+#include "HAL/RunnableThread.h"
+#include "HAL/Event.h"
+#include "Misc/ScopeLock.h"
+#include "CustomConfigFile.h"
+#include "Utility.h"
+
+FPerforceMonitor::FPerforceMonitor(const TSharedRef<FPerforceConnection>& InPerforce, const FString& InBranchClientPath, const FString& InSelectedClientFileName, const FString& InSelectedProjectIdentifier, const FString& InLogPath)
+	: Perforce(InPerforce)
+	, BranchClientPath(InBranchClientPath)
+	, SelectedClientFileName(InSelectedClientFileName)
+	, SelectedProjectIdentifier(InSelectedProjectIdentifier)
+	, LogWriter(*InLogPath, 128 * 1024)
+	, WorkerThread(nullptr)
+	, ZippedBinariesConfigChangeNumber(0)
+	, RefreshEvent(FPlatformProcess::GetSynchEventFromPool())
+	, AbortEvent(FPlatformProcess::GetSynchEventFromPool(true))
+	, CurrentMaxChanges(0)
+	, PendingMaxChanges(100)
+	, LastChangeByCurrentUser(-1)
+	, LastCodeChangeByCurrentUser(-1)
+	, bDisposing(false)
+{
+}
+
+FPerforceMonitor::~FPerforceMonitor()
+{
+	bDisposing = true;
+	if(WorkerThread != nullptr)
+	{
+		AbortEvent->Trigger();
+		RefreshEvent->Trigger();
+
+		WorkerThread->WaitForCompletion();
+		WorkerThread = nullptr;
+	}
+	FPlatformProcess::ReturnSynchEventToPool(RefreshEvent);
+}
+
+void FPerforceMonitor::Start()
+{
+	check(WorkerThread == nullptr);
+	WorkerThread = FRunnableThread::Create(this, *FString::Printf(TEXT("PerforceMonitor-%s"), *SelectedClientFileName));
+}
+
+FString FPerforceMonitor::GetLastStatusMessage() const
+{
+	FScopeLock Lock(&CriticalSection);
+	return LastStatusMessage;
+}
+
+int FPerforceMonitor::GetCurrentMaxChanges() const
+{
+	FScopeLock Lock(&CriticalSection);
+	return CurrentMaxChanges;
+}
+
+int FPerforceMonitor::GetPendingMaxChanges() const
+{
+	FScopeLock Lock(&CriticalSection);
+	return PendingMaxChanges;
+}
+
+TArray<FString> FPerforceMonitor::GetOtherStreamNames() const
+{
+	FScopeLock Lock(&CriticalSection);
+	return OtherStreamNames;
+}
+
+int FPerforceMonitor::GetLastChangeByCurrentUser() const
+{
+	FScopeLock Lock(&CriticalSection);
+	return LastChangeByCurrentUser;
+}
+
+int FPerforceMonitor::GetLastCodeChangeByCurrentUser() const
+{
+	FScopeLock Lock(&CriticalSection);
+	return LastCodeChangeByCurrentUser;
+}
+
+bool FPerforceMonitor::HasZippedBinaries() const
+{
+	FScopeLock Lock(&CriticalSection);
+	return ChangeNumberToZippedBinaries.Num() > 0;
+}
+
+TArray<TSharedRef<FPerforceChangeSummary, ESPMode::ThreadSafe>> FPerforceMonitor::GetChanges() const
+{
+	FScopeLock Lock(&CriticalSection);
+	return Changes;
+}
+
+bool FPerforceMonitor::TryGetChangeType(int ChangeNumber, EPerforceChangeType& OutType) const
+{
+	FScopeLock Lock(&CriticalSection);
+
+	const EPerforceChangeType* ChangeType = ChangeNumberToType.Find(ChangeNumber);
+	if(ChangeType != nullptr)
+	{
+		OutType = *ChangeType;
+		return true;
+	}
+	return false;
+}
+
+bool FPerforceMonitor::TryGetArchivePathForChangeNumber(int ChangeNumber, FString& OutArchivePath) const
+{
+	FScopeLock Lock(&CriticalSection);
+
+	const FString* ArchivePath = ChangeNumberToZippedBinaries.Find(ChangeNumber);
+	if(ArchivePath != nullptr)
+	{
+		OutArchivePath = *ArchivePath;
+		return true;
+	}
+	return false;
+}
+
+TSet<int> FPerforceMonitor::GetPromotedChangeNumbers() const
+{
+	FScopeLock Lock(&CriticalSection);
+	return PromotedChangeNumbers;
+}
+
+void FPerforceMonitor::Refresh()
+{
+	RefreshEvent->Trigger();
+}
+
+uint32 FPerforceMonitor::Run()
+{
+	// TODO we should likely turn on exceptions, or handle this better
+	//try
+	{
+		RunInternal();
+	}
+	//catch(FAbortException)
+	{
+	}
+	return 0;
+}
+
+void FPerforceMonitor::RunInternal()
+{
+	FString StreamName;
+	Perforce->GetActiveStream(StreamName, AbortEvent, LogWriter);
+
+	// Try to update the zipped binaries list before anything else, because it causes a state change in the UI
+	UpdateZippedBinaries();
+
+	while(!bDisposing)
+	{
+		FDateTime StartTime = FDateTime::UtcNow();
+
+		// Check we haven't switched streams
+		if(OnStreamChange)
+		{
+			FString NewStreamName;
+			if(Perforce->GetActiveStream(NewStreamName, AbortEvent, LogWriter) && NewStreamName != StreamName)
+			{
+				OnStreamChange();
+			}
+		}
+
+		// Update the stream list
+		if(StreamName.Len() > 0)
+		{
+			TArray<FString> NewOtherStreamNames;
+			Perforce->FindStreams(FPerforceUtils::GetClientOrDepotDirectoryName(*StreamName) + "/*", NewOtherStreamNames, AbortEvent, LogWriter);
+
+			FScopeLock Lock(&CriticalSection);
+			OtherStreamNames = NewOtherStreamNames;
+		}
+
+		// Check for any p4 changes
+		if(!UpdateChanges())
+		{
+			LastStatusMessage = TEXT("Failed to update changes");
+		}
+		else if(!UpdateChangeTypes())
+		{
+			LastStatusMessage = TEXT("Failed to update change types");
+		}
+		else if(!UpdateZippedBinaries())
+		{
+			LastStatusMessage = TEXT("Failed to update zipped binaries list");
+		}
+		else
+		{
+			LastStatusMessage = FString::Printf(TEXT("Last update took %dms"), (FDateTime::UtcNow() - StartTime).GetTotalMilliseconds());
+		}
+
+		// Wait for another request, or scan for new builds after a timeout
+		RefreshEvent->Wait(FTimespan::FromSeconds(60.0));
+	}
+}
+
+bool FPerforceMonitor::UpdateChanges()
+{
+	// Get the current status of the build
+	int MaxChanges;
+	int OldestChangeNumber = -1;
+	int NewestChangeNumber = -1;
+	TSet<int> CurrentChangelists;
+	TSet<int> PrevPromotedChangelists;
+	{
+		FScopeLock Lock(&CriticalSection);
+		MaxChanges = PendingMaxChanges;
+		if(Changes.Num() > 0)
+		{
+			NewestChangeNumber = Changes[0]->Number;
+			OldestChangeNumber = Changes[Changes.Num() - 1]->Number;
+		}
+		for(const TSharedRef<FPerforceChangeSummary, ESPMode::ThreadSafe>& Change : Changes)
+		{
+			CurrentChangelists.Add(Change->Number);
+		}
+		PrevPromotedChangelists = PromotedChangeNumbers;
+	}
+
+	// Build a full list of all the paths to sync
+	TArray<FString> DepotPaths;
+	if(SelectedClientFileName.EndsWith(TEXT(".uprojectdirs")))
+	{
+		DepotPaths.Add(FString::Printf(TEXT("%s/..."), *BranchClientPath));
+	}
+	else
+	{
+		DepotPaths.Add(FString::Printf(TEXT("%s/*"), *BranchClientPath));
+		DepotPaths.Add(FString::Printf(TEXT("%s/Engine/..."), *BranchClientPath));
+		DepotPaths.Add(FString::Printf(TEXT("%s/..."), *FPerforceUtils::GetClientOrDepotDirectoryName(*SelectedClientFileName)));
+	}
+
+	// Read any new changes
+	TArray<FPerforceChangeSummary> NewChanges;
+	if(MaxChanges > CurrentMaxChanges)
+	{
+		if(!Perforce->FindChanges(DepotPaths, MaxChanges, NewChanges, AbortEvent, LogWriter))
+		{
+			return false;
+		}
+	}
+	else
+	{
+		TArray<FString> RangedDepotPaths;
+		for(const FString& DepotPath : DepotPaths)
+		{
+			RangedDepotPaths.Add(FString::Printf(TEXT("%s@>%d"), *DepotPath, NewestChangeNumber));
+		}
+		if(!Perforce->FindChanges(RangedDepotPaths, -1, NewChanges, AbortEvent, LogWriter))
+		{
+			return false;
+		}
+	}
+
+	// Remove anything we already have
+	NewChanges.RemoveAll([&CurrentChangelists](const FPerforceChangeSummary& Change){ return CurrentChangelists.Contains(Change.Number); });
+
+	// Update the change ranges
+	if(NewChanges.Num() > 0)
+	{
+		OldestChangeNumber = FMath::Max(OldestChangeNumber, NewChanges[NewChanges.Num() - 1].Number);
+		NewestChangeNumber = FMath::Min(NewestChangeNumber, NewChanges[0].Number);
+	}
+
+	// Fixup any ROBOMERGE authors
+	const FString RoboMergePrefix = TEXT("#ROBOMERGE-AUTHOR:");
+	for(FPerforceChangeSummary& Change : NewChanges)
+	{
+		if(Change.Description.StartsWith(RoboMergePrefix))
+		{
+			int StartIdx = RoboMergePrefix.Len();
+			while(StartIdx < Change.Description.Len() && Change.Description[StartIdx] == ' ')
+			{
+				StartIdx++;
+			}
+
+			int EndIdx = StartIdx;
+			while(EndIdx < Change.Description.Len() && !FChar::IsWhitespace(Change.Description[EndIdx]))
+			{
+				EndIdx++;
+			}
+
+			if(EndIdx > StartIdx)
+			{
+				Change.User = Change.Description.Mid(StartIdx, EndIdx - StartIdx);
+				Change.Description = "ROBOMERGE: " + Change.Description.Mid(EndIdx).TrimStartAndEnd();
+			}
+		}
+	}
+
+	// Process the new changes received
+	if(NewChanges.Num() > 0 || MaxChanges < CurrentMaxChanges)
+	{
+		// Insert them into the builds list
+		{
+			FScopeLock Lock(&CriticalSection);
+
+			int InsertIdx = 0;
+			for(const FPerforceChangeSummary& NewChange : NewChanges)
+			{
+				while(InsertIdx < Changes.Num() && NewChange.Number < Changes[InsertIdx]->Number)
+				{
+					InsertIdx++;
+				}
+				Changes.Insert(MakeShared<FPerforceChangeSummary, ESPMode::ThreadSafe>(NewChange), InsertIdx);
+			}
+
+			if(Changes.Num() > MaxChanges)
+			{
+				Changes.RemoveAt(MaxChanges, Changes.Num() - MaxChanges);
+			}
+
+			CurrentMaxChanges = MaxChanges;
+		}
+
+		// Find the last submitted change by the current user
+		int NewLastChangeByCurrentUser = -1;
+		for(const FChangeSharedRef& Change : Changes)
+		{
+			if(Change->User == Perforce->UserName)
+			{
+				NewLastChangeByCurrentUser = FMath::Max(NewLastChangeByCurrentUser, Change->Number);
+			}
+		}
+		LastChangeByCurrentUser = NewLastChangeByCurrentUser;
+
+		// Notify the main window that we've got more data
+		if(OnUpdate)
+		{
+			OnUpdate();
+		}
+	}
+	return true;
+}
+
+bool FPerforceMonitor::UpdateChangeTypes()
+{
+	// Find the changes we need to query
+	TArray<int> QueryChangeNumbers;
+	{
+		FScopeLock Lock(&CriticalSection);
+		for(const FChangeSharedRef& Change : Changes)
+		{
+			if(!ChangeNumberToType.Contains(Change->Number))
+			{
+				QueryChangeNumbers.Add(Change->Number);
+			}
+		}
+	}
+
+	// Update them in batches
+	for(int QueryChangeNumber : QueryChangeNumbers)
+	{
+		// If there's something to check for, find all the content changes after this changelist
+		TSharedPtr<FPerforceDescribeRecord> DescribeRecord;
+		if(Perforce->Describe(QueryChangeNumber, DescribeRecord, AbortEvent, LogWriter))
+		{
+			// Check whether the files are code or content
+			EPerforceChangeType Type = EPerforceChangeType::Content;
+			for(const FPerforceDescribeFileRecord& File : DescribeRecord->Files)
+			{
+				const FString& DepotFile = File.DepotFile;
+				if(DepotFile.EndsWith(TEXT(".cs")) || DepotFile.EndsWith(TEXT(".h")) || DepotFile.EndsWith(TEXT(".cpp")) || DepotFile.EndsWith(TEXT(".usf")))
+				{
+					Type = EPerforceChangeType::Code;
+					break;
+				}
+			}
+			// Update the type of this change
+			{
+				FScopeLock Lock(&CriticalSection);
+				if(!ChangeNumberToType.Contains(QueryChangeNumber))
+				{
+					ChangeNumberToType.Add(QueryChangeNumber, Type);
+				}
+			}
+		}
+
+		// Find the last submitted code change by the current user
+		int NewLastCodeChangeByCurrentUser = -1;
+		for(const FChangeSharedRef& Change : Changes)
+		{
+			if(Change->User == Perforce->UserName)
+			{
+				EPerforceChangeType* Type = ChangeNumberToType.Find(Change->Number);
+				if(Type != nullptr && *Type == EPerforceChangeType::Code)
+				{
+					NewLastCodeChangeByCurrentUser = FMath::Max(NewLastCodeChangeByCurrentUser, Change->Number);
+				}
+			}
+		}
+		LastCodeChangeByCurrentUser = NewLastCodeChangeByCurrentUser;
+
+		// Notify the main window that we've got an update
+		if(OnUpdateMetadata)
+		{
+			OnUpdateMetadata();
+		}
+	}
+	return true;
+}
+
+bool FPerforceMonitor::UpdateZippedBinaries()
+{
+	// Get the path to the config file
+	FString ClientConfigFileName = FPerforceUtils::GetClientOrDepotDirectoryName(*SelectedClientFileName) + TEXT("/Build/UnrealGameSync.ini");
+
+	// Find the most recent change to that file (if the file doesn't exist, we succeed and just get 0 changes).
+	TArray<FPerforceChangeSummary> ConfigFileChanges;
+	if(!Perforce->FindChanges(ClientConfigFileName, 1, ConfigFileChanges, AbortEvent, LogWriter))
+	{
+		return false;
+	}
+
+	// Update the zipped binaries path if it's changed
+	int NewZippedBinariesConfigChangeNumber = (ConfigFileChanges.Num() > 0)? ConfigFileChanges[0].Number : 0;
+	if(NewZippedBinariesConfigChangeNumber != ZippedBinariesConfigChangeNumber)
+	{
+		// Read the config file
+		FString NewZippedBinariesPath;
+		if(NewZippedBinariesConfigChangeNumber != 0)
+		{
+			TArray<FString> Lines;
+			if(!Perforce->Print(ClientConfigFileName, Lines, AbortEvent, LogWriter))
+			{
+				return false;
+			}
+
+			FCustomConfigFile NewConfigFile;
+			NewConfigFile.Parse(Lines);
+
+			TSharedPtr<const FCustomConfigSection> ProjectSection = NewConfigFile.FindSection(*SelectedProjectIdentifier);
+			if(ProjectSection.IsValid())
+			{
+				NewZippedBinariesPath = ProjectSection->GetValue(TEXT("ZippedBinariesPath"), TEXT(""));
+			}
+		}
+
+		// Update the stored settings
+		{
+			FScopeLock Lock(&CriticalSection);
+			ZippedBinariesPath = NewZippedBinariesPath;
+			ZippedBinariesConfigChangeNumber = NewZippedBinariesConfigChangeNumber;
+		}
+	}
+
+	TMap<int, FString> NewChangeNumberToZippedBinaries;
+	if(ZippedBinariesPath.Len() > 0)
+	{
+		TArray<FPerforceFileChangeSummary> ZipChanges;
+		if(!Perforce->FindFileChanges(ZippedBinariesPath, 100, ZipChanges, AbortEvent, LogWriter))
+		{
+			return false;
+		}
+
+		for(const FPerforceFileChangeSummary& ZipChange : ZipChanges)
+		{
+			if(ZipChange.Action != TEXT("purge"))
+			{
+				TArray<FString> Tokens;
+				ZipChange.Description.ParseIntoArray(Tokens, TEXT(" "));
+				if(Tokens.Num() >= 2 && Tokens[0].StartsWith(TEXT("[CL")) && Tokens[1].EndsWith(TEXT("]")))
+				{
+					int OriginalChangeNumber;
+					if(FUtility::TryParse(*Tokens[1].Mid(0, Tokens[1].Len() - 1), OriginalChangeNumber) && !NewChangeNumberToZippedBinaries.Contains(OriginalChangeNumber))
+					{
+						NewChangeNumberToZippedBinaries[OriginalChangeNumber] = FString::Printf(TEXT("%s#%d"), *ZippedBinariesPath, ZipChange.Revision);
+					}
+				}
+			}
+		}
+	}
+
+	// Check if the list of zipped binaries have changed
+	{
+		FScopeLock Lock(&CriticalSection);
+
+		bool bUpdate = ChangeNumberToZippedBinaries.Num() != NewChangeNumberToZippedBinaries.Num();
+		if(!bUpdate)
+		{
+			for(const TTuple<int, FString>& Pair : ChangeNumberToZippedBinaries)
+			{
+				const FString* Other = NewChangeNumberToZippedBinaries.Find(Pair.Key);
+				if(Other == nullptr || *Other != Pair.Value)
+				{
+					bUpdate = true;
+					break;
+				}
+			}
+		}
+
+		if(bUpdate)
+		{
+			ChangeNumberToZippedBinaries = NewChangeNumberToZippedBinaries;
+			if(OnUpdateMetadata && Changes.Num() > 0)
+			{
+				OnUpdateMetadata();
+			}
+		}
+	}
+
+	return true;
+}
+
diff --git a/Engine/Source/Programs/SlateUGS/Private/UGSCore/PerforceMonitor.h b/Engine/Source/Programs/SlateUGS/Private/UGSCore/PerforceMonitor.h
new file mode 100644
--- /dev/null
+++ b/Engine/Source/Programs/SlateUGS/Private/UGSCore/PerforceMonitor.h
@@ -0,0 +1,77 @@
+// Copyright 1998-2017 Epic Games, Inc. All Rights Reserved.
+
+#pragma once
+
+#include "CoreMinimal.h"
+#include "Misc/Timespan.h"
+#include "HAL/Runnable.h"
+#include "HAL/Event.h"
+#include "Perforce.h"
+#include "OutputAdapters.h"
+
+enum class EPerforceChangeType
+{
+	Code,
+	Content,
+};
+
+class FPerforceMonitor : FRunnable
+{
+public:
+	typedef TSharedRef<FPerforceChangeSummary, ESPMode::ThreadSafe> FChangeSharedRef;
+
+	TFunction<void()> OnUpdate;
+	TFunction<void()> OnUpdateMetadata;
+	TFunction<void()> OnStreamChange;
+
+	FPerforceMonitor(const TSharedRef<FPerforceConnection>& InPerforce, const FString& InBranchClientPath, const FString& InSelectedClientFileName, const FString& InSelectedProjectIdentifier, const FString& InLogPath);
+	~FPerforceMonitor();
+
+	void Start();
+
+	FString GetLastStatusMessage() const;
+	int GetCurrentMaxChanges() const;
+	int GetPendingMaxChanges() const;
+	TArray<FString> GetOtherStreamNames() const;
+	int GetLastChangeByCurrentUser() const;
+	int GetLastCodeChangeByCurrentUser() const;
+	bool HasZippedBinaries() const;
+
+	TArray<FChangeSharedRef> GetChanges() const;
+	bool TryGetChangeType(int ChangeNumber, EPerforceChangeType& OutType) const;
+	bool TryGetArchivePathForChangeNumber(int ChangeNumber, FString& OutArchivePath) const;
+	TSet<int> GetPromotedChangeNumbers() const;
+
+	void Refresh();
+
+private:
+	TSharedRef<FPerforceConnection> Perforce;
+	const FString BranchClientPath;
+	const FString SelectedClientFileName;
+	const FString SelectedProjectIdentifier;
+	FBoundedLogWriter LogWriter;
+	FRunnableThread* WorkerThread;
+	TArray<FChangeSharedRef> Changes; // Sorted
+	TMap<int, EPerforceChangeType> ChangeNumberToType;
+	TSet<int> PromotedChangeNumbers;
+	int ZippedBinariesConfigChangeNumber;
+	FString ZippedBinariesPath;
+	TMap<int, FString> ChangeNumberToZippedBinaries;
+	FEvent* RefreshEvent;
+	FEvent* AbortEvent;
+	FString LastStatusMessage;
+	int CurrentMaxChanges;
+	int PendingMaxChanges;
+	int LastChangeByCurrentUser;
+	int LastCodeChangeByCurrentUser;
+	TArray<FString> OtherStreamNames;
+	bool bDisposing;
+	mutable FCriticalSection CriticalSection;
+
+	virtual uint32 Run() override;
+	void RunInternal();
+
+	bool UpdateChanges();
+	bool UpdateChangeTypes();
+	bool UpdateZippedBinaries();
+};
