Fix copyright

diff --git a/Engine/Source/Programs/SlateUGS/Private/UGSCore/EventMonitor.cpp b/Engine/Source/Programs/SlateUGS/Private/UGSCore/EventMonitor.cpp
--- a/Engine/Source/Programs/SlateUGS/Private/UGSCore/EventMonitor.cpp
+++ b/Engine/Source/Programs/SlateUGS/Private/UGSCore/EventMonitor.cpp
@@ -1,4 +1,4 @@
-// Copyright 1998-2017 Epic Games, Inc. All Rights Reserved.
+// Copyright Epic Games, Inc. All Rights Reserved.
 
 #include "EventMonitor.h"
 #include "HAL/Event.h"
diff --git a/Engine/Source/Programs/SlateUGS/Private/UGSCore/EventMonitor.h b/Engine/Source/Programs/SlateUGS/Private/UGSCore/EventMonitor.h
--- a/Engine/Source/Programs/SlateUGS/Private/UGSCore/EventMonitor.h
+++ b/Engine/Source/Programs/SlateUGS/Private/UGSCore/EventMonitor.h
@@ -1,4 +1,4 @@
-// Copyright 1998-2017 Epic Games, Inc. All Rights Reserved.
+// Copyright Epic Games, Inc. All Rights Reserved.
 
 #pragma once
 
diff --git a/Engine/Source/Programs/SlateUGS/Private/UGSCore/PerforceMonitor.cpp b/Engine/Source/Programs/SlateUGS/Private/UGSCore/PerforceMonitor.cpp
--- a/Engine/Source/Programs/SlateUGS/Private/UGSCore/PerforceMonitor.cpp
+++ b/Engine/Source/Programs/SlateUGS/Private/UGSCore/PerforceMonitor.cpp
@@ -1,4 +1,4 @@
-// Copyright 1998-2017 Epic Games, Inc. All Rights Reserved.
+// Copyright Epic Games, Inc. All Rights Reserved.
 
 #include "PerforceMonitor.h"
 #include "HAL/PlatformProcess.h"
diff --git a/Engine/Source/Programs/SlateUGS/Private/UGSCore/PerforceMonitor.h b/Engine/Source/Programs/SlateUGS/Private/UGSCore/PerforceMonitor.h
--- a/Engine/Source/Programs/SlateUGS/Private/UGSCore/PerforceMonitor.h
+++ b/Engine/Source/Programs/SlateUGS/Private/UGSCore/PerforceMonitor.h
@@ -1,4 +1,4 @@
-// Copyright 1998-2017 Epic Games, Inc. All Rights Reserved.
+// Copyright Epic Games, Inc. All Rights Reserved.
 
 #pragma once
 
