Added bones of sync filter window to game sync tab when clicking the Filter button (this is the wrong button and it will be moved at a later point)

diff --git a/Engine/Source/Programs/SlateUGS/Private/UGSTab.cpp b/Engine/Source/Programs/SlateUGS/Private/UGSTab.cpp
--- a/Engine/Source/Programs/SlateUGS/Private/UGSTab.cpp
+++ b/Engine/Source/Programs/SlateUGS/Private/UGSTab.cpp
@@ -247,6 +247,16 @@ FString UGSTab::GetSyncProgress() const
 	return Workspace->GetCurrentProgress().Key;
 }
 
+const TArray<FString>& UGSTab::GetSyncFilters() const
+{
+	return WorkspaceSettings->SyncView;
+}
+
+const TArray<FString>& UGSTab::GetCombinedSyncFilter() const
+{
+	return CombinedSyncFilter;
+}
+
 // This is getting called on a thread, lock our stuff up
 void UGSTab::OnWorkspaceSyncComplete(TSharedRef<FWorkspaceUpdateContext, ESPMode::ThreadSafe> WorkspaceContext, EWorkspaceUpdateResult SyncResult, const FString& StatusMessage)
 {
@@ -324,7 +334,6 @@ void UGSTab::SetupWorkspace()
 		OnWorkspaceSyncComplete(WorkspaceContext, SyncResult, StatusMessage);
 	};
 
-	// Todo: Eventually move into the sync operation
 	CombinedSyncFilter = FUserSettings::GetCombinedSyncFilter(
 		Workspace->GetSyncCategories(),
 		UserSettings->SyncView,
diff --git a/Engine/Source/Programs/SlateUGS/Private/UGSTab.h b/Engine/Source/Programs/SlateUGS/Private/UGSTab.h
--- a/Engine/Source/Programs/SlateUGS/Private/UGSTab.h
+++ b/Engine/Source/Programs/SlateUGS/Private/UGSTab.h
@@ -25,6 +25,8 @@ public:
 	void OnSyncLatest();
 	bool IsSyncing() const;
 	FString GetSyncProgress() const;
+	const TArray<FString>& GetSyncFilters() const; // Todo: is return type okay?
+	const TArray<FString>& GetCombinedSyncFilter() const; // Todo: is return type okay?
 private:
 
 	void OnWorkspaceSyncComplete(
diff --git a/Engine/Source/Programs/SlateUGS/Private/Widgets/SErrorWindow.cpp b/Engine/Source/Programs/SlateUGS/Private/Widgets/SErrorWindow.cpp
deleted file mode 100644
diff --git a/Engine/Source/Programs/SlateUGS/Private/Widgets/SErrorWindow.h b/Engine/Source/Programs/SlateUGS/Private/Widgets/SErrorWindow.h
deleted file mode 100644
diff --git a/Engine/Source/Programs/SlateUGS/Private/Widgets/SGameSyncTab.cpp b/Engine/Source/Programs/SlateUGS/Private/Widgets/SGameSyncTab.cpp
--- a/Engine/Source/Programs/SlateUGS/Private/Widgets/SGameSyncTab.cpp
+++ b/Engine/Source/Programs/SlateUGS/Private/Widgets/SGameSyncTab.cpp
@@ -7,6 +7,7 @@
 #include "SPositiveActionButton.h"
 #include "SSimpleComboButton.h"
 #include "SSimpleButton.h"
+#include "SSyncFilterWindow.h"
 #include "Widgets/SLogWidget.h"
 #include "Widgets/Layout/SHeader.h"
 #include "Widgets/Testing/STestSuite.h"
@@ -14,7 +15,6 @@
 #include "Widgets/Colors/SSimpleGradient.h"
 #include "Widgets/Images/SThrobber.h"
 
-
 #include "Styling/AppStyle.h"
 
 #include "UGSTab.h"
@@ -226,6 +226,7 @@ void SGameSyncTab::Construct(const FArguments& InArgs)
 						SNew(SSimpleButton)
 						.Text(LOCTEXT("Filter", "Filter"))
 						.Icon(FAppStyle::Get().GetBrush("Icons.Filter"))
+						.OnClicked_Lambda([this] { FSlateApplication::Get().AddModalWindow(SNew(SSyncFilterWindow).Tab(Tab), Tab->GetTabArgs().GetOwnerWindow(), false); return FReply::Handled(); })
 					]
 					+SHorizontalBox::Slot()
 					.AutoWidth()
diff --git a/Engine/Source/Programs/SlateUGS/Private/Widgets/SPopupTextWindow.cpp b/Engine/Source/Programs/SlateUGS/Private/Widgets/SPopupTextWindow.cpp
new file mode 100644
--- /dev/null
+++ b/Engine/Source/Programs/SlateUGS/Private/Widgets/SPopupTextWindow.cpp
@@ -0,0 +1,41 @@
+// Copyright Epic Games, Inc. All Rights Reserved.
+
+#include "SPopupTextWindow.h"
+
+#define LOCTEXT_NAMESPACE "PopupTextWindow"
+
+void SPopupTextWindow::Construct(const FArguments& InArgs)
+{
+	SWindow::Construct(SWindow::FArguments()
+	.Title(InArgs._TitleText)
+	.SizingRule(ESizingRule::Autosized)
+	.MaxWidth(400)
+	[
+		SNew(SVerticalBox)
+		+SVerticalBox::Slot()
+		.AutoHeight()
+		.HAlign(HAlign_Center)
+		.Padding(25.0f, 25.0f)
+		[
+			SNew(STextBlock)
+			.AutoWrapText(true)
+			.Justification(InArgs._BodyTextJustification)
+			.Text(InArgs._BodyText)
+		]
+		+SVerticalBox::Slot()
+		.AutoHeight()
+		.HAlign(HAlign_Center)
+		[
+			SNew(SButton)
+			.HAlign(HAlign_Center)
+			.Text(LOCTEXT("PopupTextWindowOkayButtonText", "Ok"))
+			.OnClicked_Lambda([this]()
+			{
+				SharedThis(this)->RequestDestroyWindow();
+				return FReply::Handled();
+			})
+		]
+	]);
+}
+
+#undef LOCTEXT_NAMESPACE
diff --git a/Engine/Source/Programs/SlateUGS/Private/Widgets/SPopupTextWindow.h b/Engine/Source/Programs/SlateUGS/Private/Widgets/SPopupTextWindow.h
new file mode 100644
--- /dev/null
+++ b/Engine/Source/Programs/SlateUGS/Private/Widgets/SPopupTextWindow.h
@@ -0,0 +1,19 @@
+// Copyright Epic Games, Inc. All Rights Reserved.
+
+#pragma once
+
+#include "Widgets/SWindow.h"
+
+class SPopupTextWindow : public SWindow
+{
+public:
+	SLATE_BEGIN_ARGS(SPopupTextWindow)
+		: _BodyTextJustification(ETextJustify::Center)
+		{}
+		SLATE_ARGUMENT(FText, TitleText)
+		SLATE_ARGUMENT(FText, BodyText)
+		SLATE_ARGUMENT(ETextJustify::Type, BodyTextJustification)
+	SLATE_END_ARGS()
+
+	void Construct(const FArguments& InArgs);
+};
diff --git a/Engine/Source/Programs/SlateUGS/Private/Widgets/SSyncFilterWindow.cpp b/Engine/Source/Programs/SlateUGS/Private/Widgets/SSyncFilterWindow.cpp
new file mode 100644
--- /dev/null
+++ b/Engine/Source/Programs/SlateUGS/Private/Widgets/SSyncFilterWindow.cpp
@@ -0,0 +1,105 @@
+// Copyright Epic Games, Inc. All Rights Reserved.
+
+#include "SSyncFilterWindow.h"
+
+#include "UGSTab.h"
+
+#include "SPopupTextWindow.h"
+#include "Widgets/Layout/SUniformGridPanel.h"
+#include "SPrimaryButton.h"
+#include "Widgets/Colors/SSimpleGradient.h" // Todo: remove
+
+#define LOCTEXT_NAMESPACE "SSyncFilterWindow"
+
+void SSyncFilterWindow::Construct(const FArguments& InArgs)
+{
+	Tab = InArgs._Tab;
+
+	SWindow::Construct(SWindow::FArguments()
+	.Title(LOCTEXT("WindowTitle", "Sync Filters"))
+	.SizingRule(ESizingRule::FixedSize)
+	.ClientSize(FVector2D(1100, 800))
+	[
+		SNew(SVerticalBox)
+		// Hint text
+		+SVerticalBox::Slot()
+		.AutoHeight()
+		.Padding(20.0f, 20.0f)
+		[
+			SNew(STextBlock)
+			.Text(LOCTEXT("SyncFilterHintText", "Files synced from Perforce may be filtered by a custom stream view, and list of predefined categories. Settings for the current workspace override defaults for all workspaces."))
+		]
+		// Filter checkbox list
+		+SVerticalBox::Slot()
+		.Padding(20.0f, 0.0f)
+		[
+			SNew(SSimpleGradient)
+			.StartColor(FLinearColor(161.0f / 255.0f, 57.0f / 255.0f, 191.0f / 255.0f))
+			.EndColor(FLinearColor(100.0f / 255.0f, 100.0f / 255.0f, 100.0f / 255.0f))
+		]
+		// Buttons
+		+SVerticalBox::Slot()
+		.AutoHeight()
+		.Padding(0.0f, 20.0f)
+		[
+			SNew(SHorizontalBox)
+			+SHorizontalBox::Slot()
+			.Padding(20.0f, 0.0f, 0.0f, 0.0f)
+			.HAlign(HAlign_Left)
+			[
+				SNew(SButton)
+				.HAlign(HAlign_Center)
+				.Text(LOCTEXT("ShowCombinedFilterButtonText", "Show Combined Filter"))
+				.OnClicked(this, &SSyncFilterWindow::OnShowCombinedFilterClicked)
+			]
+			+SHorizontalBox::Slot()
+			.Padding(0.0f, 0.0f, 10.0f, 0.0f)
+			.HAlign(HAlign_Right)
+			[
+				SNew(SUniformGridPanel)
+				.SlotPadding(FMargin(10.0f, 0.0f))
+				+SUniformGridPanel::Slot(0, 0)
+				[
+					SNew(SPrimaryButton)
+					.Text(LOCTEXT("OkButtonText", "Ok"))
+					.OnClicked(this, &SSyncFilterWindow::OnOkClicked)
+				]
+				+SUniformGridPanel::Slot(1, 0)
+				[
+					SNew(SButton)
+					.HAlign(HAlign_Center)
+					.Text(LOCTEXT("CancelButtonText", "Cancel"))
+					.OnClicked(this, &SSyncFilterWindow::OnCancelClicked)
+				]
+			]
+		]
+	]);
+}
+
+FReply SSyncFilterWindow::OnShowCombinedFilterClicked()
+{
+	fprintf(stderr, "Sync filters:\n%sEnd of sync filters\n", FString::Join(Tab->GetSyncFilters(), TEXT("\n")));
+
+	TSharedRef<SPopupTextWindow> CombinedFilterWindow = SNew(SPopupTextWindow)
+		.TitleText(FText::FromString("Combined Sync Filter"))
+		.BodyText(FText::FromString(FString::Join(Tab->GetCombinedSyncFilter(), TEXT("\n"))))
+		.BodyTextJustification(ETextJustify::Left);
+	FSlateApplication::Get().AddModalWindow(CombinedFilterWindow, SharedThis(this), false);
+	return FReply::Handled();
+}
+
+FReply SSyncFilterWindow::OnOkClicked()
+{
+	// Todo: save sync filters
+
+	FSlateApplication::Get().FindWidgetWindow(AsShared())->RequestDestroyWindow();
+	return FReply::Handled();
+}
+
+FReply SSyncFilterWindow::OnCancelClicked()
+{
+	FSlateApplication::Get().FindWidgetWindow(AsShared())->RequestDestroyWindow();
+	return FReply::Handled();
+}
+
+#undef LOCTEXT_NAMESPACE
diff --git a/Engine/Source/Programs/SlateUGS/Private/Widgets/SSyncFilterWindow.h b/Engine/Source/Programs/SlateUGS/Private/Widgets/SSyncFilterWindow.h
new file mode 100644
--- /dev/null
+++ b/Engine/Source/Programs/SlateUGS/Private/Widgets/SSyncFilterWindow.h
@@ -0,0 +1,24 @@
+// Copyright Epic Games, Inc. All Rights Reserved.
+
+#pragma once
+
+#include "Widgets/SWindow.h"
+
+class UGSTab;
+
+class SSyncFilterWindow : public SWindow
+{
+public:
+	SLATE_BEGIN_ARGS(SSyncFilterWindow) {}
+		SLATE_ARGUMENT(UGSTab*, Tab)
+	SLATE_END_ARGS()
+
+	void Construct(const FArguments& InArgs);
+
+private:
+	FReply OnShowCombinedFilterClicked();
+	FReply OnOkClicked();
+	FReply OnCancelClicked();
+
+	UGSTab* Tab;
+};
diff --git a/Engine/Source/Programs/SlateUGS/Private/Widgets/SWorkspaceWindow.cpp b/Engine/Source/Programs/SlateUGS/Private/Widgets/SWorkspaceWindow.cpp
--- a/Engine/Source/Programs/SlateUGS/Private/Widgets/SWorkspaceWindow.cpp
+++ b/Engine/Source/Programs/SlateUGS/Private/Widgets/SWorkspaceWindow.cpp
@@ -8,7 +8,7 @@
 #include "UGSTab.h"
 #include "HordeBuildRowInfo.h"
 #include "SGameSyncTab.h"
-#include "SErrorWindow.h"
+#include "SPopupTextWindow.h"
 
 #define LOCTEXT_NAMESPACE "UGSWorkspaceWindow"
 
@@ -177,11 +177,9 @@ void SWorkspaceWindow::Construct(const FArguments& InArgs)
 				SNew(SHorizontalBox)
 				+SHorizontalBox::Slot()
 				.Padding(10.0f, 0.0f)
-				.FillWidth(6) // Todo: figure out how to right justify below buttons without using this invisible dummy button as a space filler
+				.FillWidth(6)
 				[
-					SNew(SButton)
-					.HAlign(HAlign_Center)
-					.Visibility(EVisibility::Hidden)
+					SNew(SSpacer)
 				]
 				+SHorizontalBox::Slot()
 				.HAlign(HAlign_Fill)
@@ -217,8 +215,9 @@ FReply SWorkspaceWindow::OnOkClicked()
 	else
 	{
 		// Todo: report other errors (no workspace associated with the file)
+		FText Title = FText::FromString("Error");
 		FText Error = FText::FromString("Project file does not exist, try again");
-		FSlateApplication::Get().AddModalWindow(SNew(SErrorWindow).ErrorText(Error), SharedThis(this), false);
+		FSlateApplication::Get().AddModalWindow(SNew(SPopupTextWindow).TitleText(Title).BodyText(Error), SharedThis(this), false);
 	}
 
 	return FReply::Handled();
