Fixed error caused by printing incorrectly

diff --git a/Engine/Source/Programs/SlateUGS/Private/Widgets/SSyncFilterWindow.cpp b/Engine/Source/Programs/SlateUGS/Private/Widgets/SSyncFilterWindow.cpp
--- a/Engine/Source/Programs/SlateUGS/Private/Widgets/SSyncFilterWindow.cpp
+++ b/Engine/Source/Programs/SlateUGS/Private/Widgets/SSyncFilterWindow.cpp
@@ -78,7 +78,7 @@ void SSyncFilterWindow::Construct(const FArguments& InArgs)
 
 FReply SSyncFilterWindow::OnShowCombinedFilterClicked()
 {
-	fprintf(stderr, "Sync filters:\n%sEnd of sync filters\n", FString::Join(Tab->GetSyncFilters(), TEXT("\n")));
+	fprintf(stderr, "Sync filters:\n%sEnd of sync filters\n", TCHAR_TO_ANSI(*FString::Join(Tab->GetSyncFilters(), TEXT("\n"))));
 
 	TSharedRef<SPopupTextWindow> CombinedFilterWindow = SNew(SPopupTextWindow)
 		.TitleText(FText::FromString("Combined Sync Filter"))
