Fixed sizing issues on combined sync filter window

diff --git a/Engine/Source/Programs/SlateUGS/Private/UGSTab.cpp b/Engine/Source/Programs/SlateUGS/Private/UGSTab.cpp
--- a/Engine/Source/Programs/SlateUGS/Private/UGSTab.cpp
+++ b/Engine/Source/Programs/SlateUGS/Private/UGSTab.cpp
@@ -281,7 +281,7 @@ void UGSTab::SetupWorkspace()
 	DetectSettings = MakeShared<FDetectProjectSettingsTask>(MakeShared<FPerforceConnection>(TEXT(""), TEXT(""), TEXT("")), ProjectFileName, MakeShared<FLineWriter>());
 
 	TSharedRef<FModalTaskResult> Result = ExecuteModalTask(
-		TabWidget,
+		FSlateApplication::Get().GetActiveModalWindow(),
 		DetectSettings.ToSharedRef(),
 		LOCTEXT("OpeningProjectTitle", "Opening Project"),
 		LOCTEXT("OpeningProjectCaption", "Opening project, please wait..."));
diff --git a/Engine/Source/Programs/SlateUGS/Private/Widgets/SPopupTextWindow.cpp b/Engine/Source/Programs/SlateUGS/Private/Widgets/SPopupTextWindow.cpp
--- a/Engine/Source/Programs/SlateUGS/Private/Widgets/SPopupTextWindow.cpp
+++ b/Engine/Source/Programs/SlateUGS/Private/Widgets/SPopupTextWindow.cpp
@@ -2,32 +2,48 @@
 
 #include "SPopupTextWindow.h"
 
+#include "Styling/AppStyle.h"
+
+#include "Widgets/Input/SMultiLineEditableTextBox.h"
+#include "SPrimaryButton.h"
+
 #define LOCTEXT_NAMESPACE "PopupTextWindow"
 
 void SPopupTextWindow::Construct(const FArguments& InArgs)
 {
+	TSharedRef<SScrollBar> InvisibleHorizontalScrollbar = SNew(SScrollBar)
+		.AlwaysShowScrollbar(false)
+		.Orientation(Orient_Horizontal);
+
 	SWindow::Construct(SWindow::FArguments()
 	.Title(InArgs._TitleText)
 	.SizingRule(ESizingRule::Autosized)
-	.MaxWidth(400)
+	.MaxWidth(600)
+	.MaxHeight(600)
 	[
 		SNew(SVerticalBox)
 		+SVerticalBox::Slot()
-		.AutoHeight()
+		.Padding(20.0f, 10.0f)
 		.HAlign(HAlign_Center)
-		.Padding(25.0f, 25.0f)
+		.VAlign(VAlign_Center)
 		[
-			SNew(STextBlock)
+			SNew(SMultiLineEditableTextBox)
+			.Padding(10.0f)
+			.AlwaysShowScrollbars(InArgs._ShowScrollBars)
+			.HScrollBar(InvisibleHorizontalScrollbar)
+			.IsReadOnly(true)
 			.AutoWrapText(true)
+			.BackgroundColor(FLinearColor::Transparent)
 			.Justification(InArgs._BodyTextJustification)
 			.Text(InArgs._BodyText)
 		]
 		+SVerticalBox::Slot()
+		.Padding(0.0f, 0.0f, 0.0f, 15.0f)
 		.AutoHeight()
 		.HAlign(HAlign_Center)
+		.VAlign(VAlign_Top)
 		[
-			SNew(SButton)
-			.HAlign(HAlign_Center)
+			SNew(SPrimaryButton)
 			.Text(LOCTEXT("PopupTextWindowOkayButtonText", "Ok"))
 			.OnClicked_Lambda([this]()
 			{
diff --git a/Engine/Source/Programs/SlateUGS/Private/Widgets/SPopupTextWindow.h b/Engine/Source/Programs/SlateUGS/Private/Widgets/SPopupTextWindow.h
--- a/Engine/Source/Programs/SlateUGS/Private/Widgets/SPopupTextWindow.h
+++ b/Engine/Source/Programs/SlateUGS/Private/Widgets/SPopupTextWindow.h
@@ -9,10 +9,12 @@ class SPopupTextWindow : public SWindow
 public:
 	SLATE_BEGIN_ARGS(SPopupTextWindow)
 		: _BodyTextJustification(ETextJustify::Center)
+		, _ShowScrollBars(false)
 		{}
 		SLATE_ARGUMENT(FText, TitleText)
 		SLATE_ARGUMENT(FText, BodyText)
 		SLATE_ARGUMENT(ETextJustify::Type, BodyTextJustification)
+		SLATE_ARGUMENT(bool, ShowScrollBars)
 	SLATE_END_ARGS()
 
 	void Construct(const FArguments& InArgs);
diff --git a/Engine/Source/Programs/SlateUGS/Private/Widgets/SSyncFilterWindow.cpp b/Engine/Source/Programs/SlateUGS/Private/Widgets/SSyncFilterWindow.cpp
--- a/Engine/Source/Programs/SlateUGS/Private/Widgets/SSyncFilterWindow.cpp
+++ b/Engine/Source/Programs/SlateUGS/Private/Widgets/SSyncFilterWindow.cpp
@@ -83,7 +83,8 @@ FReply SSyncFilterWindow::OnShowCombinedFilterClicked()
 	TSharedRef<SPopupTextWindow> CombinedFilterWindow = SNew(SPopupTextWindow)
 		.TitleText(FText::FromString("Combined Sync Filter"))
 		.BodyText(FText::FromString(FString::Join(Tab->GetCombinedSyncFilter(), TEXT("\n"))))
-		.BodyTextJustification(ETextJustify::Left);
+		.BodyTextJustification(ETextJustify::Left)
+		.ShowScrollBars(true);
 	FSlateApplication::Get().AddModalWindow(CombinedFilterWindow, SharedThis(this), false);
 	return FReply::Handled();
 }
