Add support for a MessageQueue that allows for queuing up functions from threads that will be called on the main thread
Add inital support for PerforceMonitor callbacks, and lists the latest 100 CLs for the given workspace
Add tick function to tabs, so we can clear our message queue, or other things when needed on the main thread

diff --git a/Engine/Source/Programs/SlateUGS/Private/SlateUGSApp.cpp b/Engine/Source/Programs/SlateUGS/Private/SlateUGSApp.cpp
--- a/Engine/Source/Programs/SlateUGS/Private/SlateUGSApp.cpp
+++ b/Engine/Source/Programs/SlateUGS/Private/SlateUGSApp.cpp
@@ -55,6 +55,7 @@ int RunSlateUGS(const TCHAR* CommandLine)
 			FTSTicker::GetCoreTicker().Tick(FApp::GetDeltaTime());
 			FSlateApplication::Get().PumpMessages();
 			FSlateApplication::Get().Tick();
+			TabManager.Tick();
 			FPlatformProcess::Sleep(0.01f);
 
 			GFrameCounter++;
diff --git a/Engine/Source/Programs/SlateUGS/Private/UGSCore/Perforce.cpp b/Engine/Source/Programs/SlateUGS/Private/UGSCore/Perforce.cpp
--- a/Engine/Source/Programs/SlateUGS/Private/UGSCore/Perforce.cpp
+++ b/Engine/Source/Programs/SlateUGS/Private/UGSCore/Perforce.cpp
@@ -1366,17 +1366,6 @@ bool FPerforceConnection::RunCommandWithBinaryOutput(const FString& CommandLine,
 
 	FProcHandle P4Proc = FPlatformProcess::CreateProc(*GPerforceExe.PerforceExe, *FullCommandLine, false, false, false, &ProcId, 0, nullptr, WritePipe, ReadPipe);
 
-	/*
-	FString P4Output;
-    FString LatestOutput = FPlatformProcess::ReadPipe(ReadPipe);
-    while (FPlatformProcess::IsProcRunning(P4Proc) || !LatestOutput.IsEmpty())
-    {
-        P4Output += LatestOutput;
-        LatestOutput = FPlatformProcess::ReadPipe(ReadPipe);
-        FPlatformProcess::Sleep(0);
-    }
-	*/
-
 	size_t BufferPos = 0;
 	size_t BufferEnd = 0;
 	TArray<uint8> Buffer;
@@ -1389,19 +1378,11 @@ bool FPerforceConnection::RunCommandWithBinaryOutput(const FString& CommandLine,
 		// Check that we don't have an abort event
 		if (AbortEvent->Wait(FTimespan::Zero()))
 		{
-			abort();
+			return false;
 			//throw FAbortException();
 		}
 
-		// Prepare space in the output buffer for the piped data
-		//size_t BytesToRead = 512;
-		//if(Buffer.Num() < BufferEnd + BytesToRead)
-		//{
-			//Buffer.AddUninitialized((BufferEnd + BytesToRead) - Buffer.Num());
-		//}
-
 		// Read data from the child process
-		//size_t BytesRead = Process.Read(Buffer.GetData() + BufferEnd, BytesToRead);
 		TArray<uint8> TempBuffer;
 		bool bRead = FPlatformProcess::ReadPipeToArray(ReadPipe, TempBuffer);
 		size_t BytesRead = TempBuffer.Num();
@@ -1418,7 +1399,7 @@ bool FPerforceConnection::RunCommandWithBinaryOutput(const FString& CommandLine,
 			}
 			else if(AbortEvent->Wait(FTimespan::FromMilliseconds(50)))
 			{
-				abort();
+				return false;
 				//throw FAbortException();
 			}
 		}
@@ -1458,84 +1439,6 @@ bool FPerforceConnection::RunCommandWithBinaryOutput(const FString& CommandLine,
 
 	checkf(BufferPos == BufferEnd, TEXT("%d bytes of incomplete record data received from P4"), BufferEnd - BufferPos);
 	return FPlatformProcess::GetProcReturnCode(P4Proc, &ExitCode) == 0;
-
-	/*
-	FChildProcess Process;
-	if(!Process.Start(PerforceExe, *FString::Printf(TEXT("-G %s"), *CommandLine), nullptr, EChildProcessOptions::OutputPipe, EChildProcessPriority::Normal))
-	{
-		return false;
-	}
-
-	size_t BufferPos = 0;
-	size_t BufferEnd = 0;
-	TArray<uint8> Buffer;
-	TMap<FString, FString> Record;
-	for(;;)
-	{
-		// Determine if the process has exited. Do this before reading, so we'll only exit the loop if no data is read AND the process has exited.
-		bool bHasExited = Process.HasExited();
-
-		// Check that we don't have an abort event
-		if(AbortEvent->Wait(FTimespan::Zero()))
-		{
-			throw FAbortException();
-		}
-
-		// Prepare space in the output buffer for the piped data
-		size_t BytesToRead = 512;
-		if(Buffer.Num() < BufferEnd + BytesToRead)
-		{
-			Buffer.AddUninitialized((BufferEnd + BytesToRead) - Buffer.Num());
-		}
-
-		// Read data from the child process
-		size_t BytesRead = Process.Read(Buffer.GetData() + BufferEnd, BytesToRead);
-		if(BytesRead == 0)
-		{
-			// If it exited, quit. Otherwise sleep until data is available.
-			if(bHasExited)
-			{
-				break;
-			}
-			else if(AbortEvent->Wait(FTimespan::FromMilliseconds(50)))
-			{
-				throw FAbortException();
-			}
-		}
-		else
-		{
-			// Add the read bytes to the buffer
-			BufferEnd += BytesRead;
-
-			// Process all the records in the buffer that we can
-			const uint8* RecordPos = Buffer.GetData() + BufferPos;
-			const uint8* RecordMaxPos = Buffer.GetData() + BufferEnd;
-			for(;;)
-			{
-				const uint8* RecordEnd = ReadBinaryRecord(RecordPos, RecordMaxPos, &Record);
-				if(RecordEnd == nullptr)
-				{
-					break;
-				}
-				HandleOutput(Record);
-				RecordPos = RecordEnd;
-			}
-			BufferPos = RecordPos - Buffer.GetData();
-
-			// Shrink the buffer down if we've got spare space
-			if(BufferPos > 64)
-			{
-				BufferEnd -= BufferPos;
-				memmove(Buffer.GetData(), Buffer.GetData() + BufferPos, BufferEnd);
-				BufferPos = 0;
-			}
-		}
-	}
-	checkf(BufferPos == BufferEnd, TEXT("%d bytes of incomplete record data received from P4"), BufferEnd - BufferPos);
-	return Process.GetExitCode() == 0;
-	*/
-
-	return false;
 }
 
 //// FPerforceUtils ////
diff --git a/Engine/Source/Programs/SlateUGS/Private/UGSTab.cpp b/Engine/Source/Programs/SlateUGS/Private/UGSTab.cpp
--- a/Engine/Source/Programs/SlateUGS/Private/UGSTab.cpp
+++ b/Engine/Source/Programs/SlateUGS/Private/UGSTab.cpp
@@ -5,6 +5,11 @@
 #include "UGSLog.h"
 #include "UGSCore/Utility.h"
 #include "UGSCore/BuildStep.h"
+#include "UGSCore/DetectProjectSettingsTask.h"
+#include "UGSCore/UserSettings.h"
+#include "UGSCore/PerforceMonitor.h"
+#include "UGSCore/EventMonitor.h"
+
 #include "Widgets/SModalTaskWindow.h"
 #include "Widgets/SLogWidget.h"
 
@@ -13,7 +18,8 @@
 UGSTab::UGSTab() : TabArgs(nullptr, FTabId()),
 				   TabWidget(SNew(SDockTab)),
 				   EmptyTabView(SNew(SEmptyTab).Tab(this)),
-				   GameSyncTabView(SNew(SGameSyncTab).Tab(this))
+				   GameSyncTabView(SNew(SGameSyncTab).Tab(this)),
+				   bHasQueuedMessages(false)
 {
 	TabWidget->SetContent(EmptyTabView);
 	TabWidget->SetLabel(FText(LOCTEXT("TabName", "Select a Project")));
@@ -237,6 +243,30 @@ void UGSTab::OnSyncLatest()
 	FPlatformProcess::ReturnSynchEventToPool(AbortEvent);
 }
 
+void UGSTab::QueueMessageForMainThread(TFunction<void()> Function)
+{
+	FScopeLock Lock(&CriticalSection);
+
+	MessageQueue.Add(Function);
+	bHasQueuedMessages = true;
+}
+
+void UGSTab::Tick()
+{
+	if (bHasQueuedMessages)
+	{
+		bHasQueuedMessages = false;
+
+		FScopeLock Lock(&CriticalSection);
+		for (const TFunction<void()>& MessageFunction : MessageQueue)
+		{
+			MessageFunction();
+		}
+
+		MessageQueue.Empty();
+	}
+}
+
 bool UGSTab::IsSyncing() const
 {
 	return Workspace->IsBusy();
@@ -270,6 +300,11 @@ void UGSTab::OnWorkspaceSyncComplete(TSharedRef<FWorkspaceUpdateContext, ESPMode
 	// TODO check this is valid, may be off
 	WorkspaceSettings->LastSyncDurationSeconds = (FDateTime::UtcNow() - WorkspaceContext->StartTime).GetSeconds();
 
+	// Queue up setting the changelist text on the main thread
+	QueueMessageForMainThread([this] {
+		GameSyncTabView->SetChangelistText(FText::FromString(FString::FromInt(WorkspaceSettings->CurrentChangeNumber)));
+	});
+
 	UserSettings->Save();
 }
 
@@ -363,6 +398,41 @@ void UGSTab::SetupWorkspace()
 	{
 		Options |= EWorkspaceUpdateOptions::OpenSolutionAfterSync;
 	}
+
+	// Setup our Perforce and Event monitoring threads
+	FString BranchClientPath = DetectSettings->BranchDirectoryName;
+	FString SelectedClientFileName = DetectSettings->NewSelectedClientFileName;
+	FString SelectedProjectIdentifier = DetectSettings->NewSelectedProjectIdentifier;
+
+	// TODO create callback functions that will be queued for the main thread to generate and update the main table view
+	PerforceMonitor = MakeShared<FPerforceMonitor>(PerforceClient.ToSharedRef(), BranchClientPath, SelectedClientFileName, SelectedProjectIdentifier, ProjectLogBaseName + ".p4.log");
+	PerforceMonitor->OnUpdate = [this]{ //printf("PerforceMonitor->OnUpdate\n"); }; //MessageQueue.Add([this]{ UpdateBuildList(); }); };
+		TArray<TSharedRef<FPerforceChangeSummary, ESPMode::ThreadSafe>> Changes = PerforceMonitor->GetChanges();
+
+		for(int ChangeIdx = 0; ChangeIdx < Changes.Num(); ChangeIdx++)
+		{
+			const FPerforceChangeSummary& Change = Changes[ChangeIdx].Get();
+			TArray<FString> DescLines;
+
+			// split on \n so we can just take the first line
+			Change.Description.ParseIntoArray(DescLines, TEXT("\n"));
+
+			printf("CL %i Author: %s@%s :: %s\n", Change.Number, TCHAR_TO_ANSI(*Change.User), TCHAR_TO_ANSI(*Change.Client), TCHAR_TO_ANSI(*DescLines[0]));
+		}
+	};
+
+	//PerforceMonitor->OnUpdateMetadata = [this]{ printf("PerforceMonitor->OnUpdateMetadata\n"); }; //MessageQueue.Add([this]{ UpdateBuildMetadata(); }); };
+	PerforceMonitor->OnStreamChange = [this]{ printf("PerforceMonitor->OnStreamChange\n"); }; // MessageQueue.Add([this]{ Owner->StreamChanged(this); }); };
+
+	/* TODO figure out if this is even working, and if so how to correctly use this
+	FString SqlConnectionString;
+	EventMonitor = MakeShared<FEventMonitor>(SqlConnectionString, FPerforceUtils::GetClientOrDepotDirectoryName(*SelectedProjectIdentifier), PerforceClient->UserName, ProjectLogBaseName + ".review.log");
+	EventMonitor->OnUpdatesReady = [this]{ printf("EventMonitor->OnUpdatesReady\n"); }; //MessageQueue.Add([this]{ UpdateReviews(); }); };
+	*/
+
+	// Start the threads
+	PerforceMonitor->Start();
+	EventMonitor->Start();
 }
 
 #undef LOCTEXT_NAMESPACE
diff --git a/Engine/Source/Programs/SlateUGS/Private/UGSTab.h b/Engine/Source/Programs/SlateUGS/Private/UGSTab.h
--- a/Engine/Source/Programs/SlateUGS/Private/UGSTab.h
+++ b/Engine/Source/Programs/SlateUGS/Private/UGSTab.h
@@ -7,8 +7,13 @@
 #include "Widgets/SWorkspaceWindow.h"
 
 #include "UGSCore/Workspace.h"
-#include "UGSCore/DetectProjectSettingsTask.h"
-#include "UGSCore/UserSettings.h"
+
+struct FUserWorkspaceSettings;
+struct FUserProjectSettings;
+struct FUserSettings;
+class FDetectProjectSettingsTask;
+class FPerforceMonitor;
+class FEventMonitor;
 
 class UGSTab
 {
@@ -20,6 +25,8 @@ public:
 	void SetTabArgs(FSpawnTabArgs InTabArgs);
 	FSpawnTabArgs GetTabArgs() const;
 
+	void Tick();
+
 	// Slate callbacks
 	bool OnWorkspaceChosen(const FString& Project);
 	void OnSyncLatest();
@@ -34,6 +41,9 @@ private:
 		EWorkspaceUpdateResult SyncResult,
 		const FString& StatusMessage);
 
+	// Allows the queuing of functions from threads to be run on the main thread
+	void QueueMessageForMainThread(TFunction<void()> Function);
+
 	// Core functions
 	void SetupWorkspace();
 
@@ -55,4 +65,12 @@ private:
 	TSharedPtr<FDetectProjectSettingsTask> DetectSettings;
 	TArray<FString> CombinedSyncFilter;
 	TSharedPtr<FUserSettings> UserSettings;
+
+	// Monitoring threads
+	TSharedPtr<FPerforceMonitor> PerforceMonitor;
+	TSharedPtr<FEventMonitor> EventMonitor;
+
+	// Queue for handling callbacks from multiple threads that need to be called on the main thread
+	TArray<TFunction<void()>> MessageQueue;
+	std::atomic<bool> bHasQueuedMessages;
 };
diff --git a/Engine/Source/Programs/SlateUGS/Private/UGSTabManager.cpp b/Engine/Source/Programs/SlateUGS/Private/UGSTabManager.cpp
--- a/Engine/Source/Programs/SlateUGS/Private/UGSTabManager.cpp
+++ b/Engine/Source/Programs/SlateUGS/Private/UGSTabManager.cpp
@@ -36,6 +36,14 @@ void UGSTabManager::ConstructTabs()
 	FGlobalTabmanager::Get()->RestoreFrom(Layout, nullptr);
 }
 
+void UGSTabManager::Tick()
+{
+	for (UGSTab& Tab : Tabs)
+	{
+		Tab.Tick();
+	}
+}
+
 TSharedRef<SDockTab> UGSTabManager::SpawnTab(int Index, const FSpawnTabArgs& Args)
 {
 	Tabs[Index].SetTabArgs(Args);
diff --git a/Engine/Source/Programs/SlateUGS/Private/UGSTabManager.h b/Engine/Source/Programs/SlateUGS/Private/UGSTabManager.h
--- a/Engine/Source/Programs/SlateUGS/Private/UGSTabManager.h
+++ b/Engine/Source/Programs/SlateUGS/Private/UGSTabManager.h
@@ -9,10 +9,14 @@ class UGSTabManager
 public:
 	UGSTabManager() = default;
 	void ConstructTabs();
-	TSharedRef<SDockTab> SpawnTab(int Index, const FSpawnTabArgs& Args);
+
+	void Tick();
+
 	bool CanSpawnTab() const;
 	int GetMaxTabs() const;
+	TSharedRef<SDockTab> SpawnTab(int Index, const FSpawnTabArgs& Args);
+
 private:
 	static constexpr int MaxTabs = 1;
 	TStaticArray<UGSTab, MaxTabs> Tabs;
-};
\ No newline at end of file
+};
