Commented out work-in-progress code

diff --git a/Engine/Source/Programs/SlateUGS/Private/UGSTab.cpp b/Engine/Source/Programs/SlateUGS/Private/UGSTab.cpp
--- a/Engine/Source/Programs/SlateUGS/Private/UGSTab.cpp
+++ b/Engine/Source/Programs/SlateUGS/Private/UGSTab.cpp
@@ -432,7 +432,7 @@ void UGSTab::SetupWorkspace()
 
 	// Start the threads
 	PerforceMonitor->Start();
-	EventMonitor->Start();
+	// EventMonitor->Start();
 }
 
 #undef LOCTEXT_NAMESPACE
