Add more filtering for the build list, we should now be correctly filtering regex'ed info, as well as other things

diff --git a/Engine/Source/Programs/SlateUGS/Private/UGSTab.cpp b/Engine/Source/Programs/SlateUGS/Private/UGSTab.cpp
--- a/Engine/Source/Programs/SlateUGS/Private/UGSTab.cpp
+++ b/Engine/Source/Programs/SlateUGS/Private/UGSTab.cpp
@@ -13,6 +13,8 @@
 #include "Widgets/SModalTaskWindow.h"
 #include "Widgets/SLogWidget.h"
 
+#include "Internationalization/Regex.h"
+
 #define LOCTEXT_NAMESPACE "UGSTab"
 
 UGSTab::UGSTab() : TabArgs(nullptr, FTabId()),
@@ -267,6 +269,98 @@ void UGSTab::Tick()
 	}
 }
 
+namespace
+{
+	bool ShouldShowChange(const FPerforceChangeSummary& Change, const TArray<FString>& ExcludeChanges)
+	{
+		for (const FString& ExcludeChange : ExcludeChanges)
+		{
+			FRegexPattern RegexPattern(*ExcludeChange, ERegexPatternFlags::CaseInsensitive);
+			FRegexMatcher RegexMatcher(RegexPattern, Change.Description);
+
+			if (RegexMatcher.FindNext())
+			{
+				return false;
+			}
+		}
+
+		if (Change.User == TEXT("buildmachine") && Change.Description.Contains(TEXT("lightmaps")))
+		{
+			return false;
+		}
+
+		return true;
+	}
+}
+
+bool UGSTab::ShouldIncludeInReviewedList(const TSet<int>& PromotedChangeNumbers, int ChangeNumber) const
+{
+	if (PromotedChangeNumbers.Contains(ChangeNumber))
+	{
+		return true;
+	}
+
+	TSharedPtr<FEventSummary> Review;
+	if (EventMonitor->TryGetSummaryForChange(ChangeNumber, Review))
+	{
+		if (Review->LastStarReview.IsValid() && Review->LastStarReview->Type == EEventType::Starred)
+		{
+			return true;
+		}
+
+		if (Review->Verdict == EReviewVerdict::Good || Review->Verdict == EReviewVerdict::Mixed)
+		{
+			return true;
+		}
+	}
+
+	return false;
+}
+
+void UGSTab::UpdateGameTabBuildList()
+{
+	TArray<TSharedRef<FPerforceChangeSummary, ESPMode::ThreadSafe>> Changes = PerforceMonitor->GetChanges();
+
+	// do we need to store these for something?
+	TSet<int> PromotedChangeNumbers = PerforceMonitor->GetPromotedChangeNumbers();
+
+	TArray<FString> ExcludeChanges;
+    TSharedPtr<FCustomConfigFile, ESPMode::ThreadSafe> ProjectConfigFile = Workspace->GetProjectConfigFile();
+    if (ProjectConfigFile.IsValid())
+    {
+        ProjectConfigFile->TryGetValues(TEXT("Options.ExcludeChanges"), ExcludeChanges);
+    }
+
+	bool bFirstChange = true;
+	bool bOnlyShowReviewed = false;
+	//  bool bOnlyShowReviewed = OnlyShowReviewedCheckBox.Checked;
+
+	for (int ChangeIdx = 0; ChangeIdx < Changes.Num(); ChangeIdx++)
+	{
+		const FPerforceChangeSummary& Change = Changes[ChangeIdx].Get();
+        if (ShouldShowChange(Change, ExcludeChanges) || PromotedChangeNumbers.Contains(Change.Number))
+        {
+			if (!bOnlyShowReviewed || (!EventMonitor->IsUnderInvestigation(Change.Number) && (ShouldIncludeInReviewedList(PromotedChangeNumbers, Change.Number) || bFirstChange)))
+			{
+				bFirstChange = false;
+
+				FDateTime DisplayTime = Change.Date;
+				if (UserSettings->bShowLocalTimes)
+				{
+					DisplayTime = (DisplayTime - DetectSettings->ServerTimeZone).GetDate();
+				}
+
+				TArray<FString> DescLines;
+
+				// split on \n so we can just take the first line
+				Change.Description.ParseIntoArray(DescLines, TEXT("\n"));
+
+				printf("%s CL %i Author: %s@%s :: %s\n", TCHAR_TO_ANSI(*DisplayTime.ToString()), Change.Number, TCHAR_TO_ANSI(*Change.User), TCHAR_TO_ANSI(*Change.Client), TCHAR_TO_ANSI(*DescLines[0]));
+			}
+		}
+	}
+}
+
 bool UGSTab::IsSyncing() const
 {
 	return Workspace->IsBusy();
@@ -406,33 +500,20 @@ void UGSTab::SetupWorkspace()
 
 	// TODO create callback functions that will be queued for the main thread to generate and update the main table view
 	PerforceMonitor = MakeShared<FPerforceMonitor>(PerforceClient.ToSharedRef(), BranchClientPath, SelectedClientFileName, SelectedProjectIdentifier, ProjectLogBaseName + ".p4.log");
-	PerforceMonitor->OnUpdate = [this]{ //printf("PerforceMonitor->OnUpdate\n"); }; //MessageQueue.Add([this]{ UpdateBuildList(); }); };
-		TArray<TSharedRef<FPerforceChangeSummary, ESPMode::ThreadSafe>> Changes = PerforceMonitor->GetChanges();
-
-		for(int ChangeIdx = 0; ChangeIdx < Changes.Num(); ChangeIdx++)
-		{
-			const FPerforceChangeSummary& Change = Changes[ChangeIdx].Get();
-			TArray<FString> DescLines;
-
-			// split on \n so we can just take the first line
-			Change.Description.ParseIntoArray(DescLines, TEXT("\n"));
-
-			printf("CL %i Author: %s@%s :: %s\n", Change.Number, TCHAR_TO_ANSI(*Change.User), TCHAR_TO_ANSI(*Change.Client), TCHAR_TO_ANSI(*DescLines[0]));
-		}
-	};
+	PerforceMonitor->OnUpdate = [this]{ QueueMessageForMainThread([this] { UpdateGameTabBuildList(); }); };
 
 	//PerforceMonitor->OnUpdateMetadata = [this]{ printf("PerforceMonitor->OnUpdateMetadata\n"); }; //MessageQueue.Add([this]{ UpdateBuildMetadata(); }); };
 	PerforceMonitor->OnStreamChange = [this]{ printf("PerforceMonitor->OnStreamChange\n"); }; // MessageQueue.Add([this]{ Owner->StreamChanged(this); }); };
 
 	/* TODO figure out if this is even working, and if so how to correctly use this
+	 */
 	FString SqlConnectionString;
 	EventMonitor = MakeShared<FEventMonitor>(SqlConnectionString, FPerforceUtils::GetClientOrDepotDirectoryName(*SelectedProjectIdentifier), PerforceClient->UserName, ProjectLogBaseName + ".review.log");
 	EventMonitor->OnUpdatesReady = [this]{ printf("EventMonitor->OnUpdatesReady\n"); }; //MessageQueue.Add([this]{ UpdateReviews(); }); };
-	*/
 
 	// Start the threads
 	PerforceMonitor->Start();
-	// EventMonitor->Start();
+	EventMonitor->Start();
 }
 
 #undef LOCTEXT_NAMESPACE
diff --git a/Engine/Source/Programs/SlateUGS/Private/UGSTab.h b/Engine/Source/Programs/SlateUGS/Private/UGSTab.h
--- a/Engine/Source/Programs/SlateUGS/Private/UGSTab.h
+++ b/Engine/Source/Programs/SlateUGS/Private/UGSTab.h
@@ -44,6 +44,9 @@ private:
 	// Allows the queuing of functions from threads to be run on the main thread
 	void QueueMessageForMainThread(TFunction<void()> Function);
 
+	bool ShouldIncludeInReviewedList(const TSet<int>& PromotedChangeNumbers, int ChangeNumber) const;
+	void UpdateGameTabBuildList();
+
 	// Core functions
 	void SetupWorkspace();
 
