Added a bit more UI to the sync filter window

diff --git a/Engine/Source/Programs/SlateUGS/Private/Widgets/SPopupTextWindow.cpp b/Engine/Source/Programs/SlateUGS/Private/Widgets/SPopupTextWindow.cpp
--- a/Engine/Source/Programs/SlateUGS/Private/Widgets/SPopupTextWindow.cpp
+++ b/Engine/Source/Programs/SlateUGS/Private/Widgets/SPopupTextWindow.cpp
@@ -18,7 +18,7 @@ void SPopupTextWindow::Construct(const FArguments& InArgs)
 	SWindow::Construct(SWindow::FArguments()
 	.Title(InArgs._TitleText)
 	.SizingRule(ESizingRule::Autosized)
-	.MaxWidth(600)
+	.MaxWidth(800)
 	.MaxHeight(600)
 	[
 		SNew(SVerticalBox)
diff --git a/Engine/Source/Programs/SlateUGS/Private/Widgets/SSyncFilterWindow.cpp b/Engine/Source/Programs/SlateUGS/Private/Widgets/SSyncFilterWindow.cpp
--- a/Engine/Source/Programs/SlateUGS/Private/Widgets/SSyncFilterWindow.cpp
+++ b/Engine/Source/Programs/SlateUGS/Private/Widgets/SSyncFilterWindow.cpp
@@ -4,10 +4,12 @@
 
 #include "UGSTab.h"
 
-#include "SPopupTextWindow.h"
+#include "Widgets/Input/SMultiLineEditableTextBox.h"
 #include "Widgets/Layout/SUniformGridPanel.h"
+#include "Widgets/Layout/SHeader.h"
 #include "SPrimaryButton.h"
-#include "Widgets/Colors/SSimpleGradient.h" // Todo: remove
+#include "SSimpleButton.h"
+#include "SPopupTextWindow.h"
 
 #define LOCTEXT_NAMESPACE "SSyncFilterWindow"
 
@@ -15,6 +17,10 @@ void SSyncFilterWindow::Construct(const FArguments& InArgs)
 {
 	Tab = InArgs._Tab;
 
+	TSharedRef<SScrollBar> InvisibleHorizontalScrollbar = SNew(SScrollBar) // Todo: this code is duplicated in SPopupTextWindow.cpp
+		.AlwaysShowScrollbar(false)										   // Is there a simpler way to make the horizontal scroll bar invisible?
+		.Orientation(Orient_Horizontal);								   // If not, I guess we should factor this out into a tiny widget?
+
 	SWindow::Construct(SWindow::FArguments()
 	.Title(LOCTEXT("WindowTitle", "Sync Filters"))
 	.SizingRule(ESizingRule::FixedSize)
@@ -33,9 +39,80 @@ void SSyncFilterWindow::Construct(const FArguments& InArgs)
 		+SVerticalBox::Slot()
 		.Padding(20.0f, 0.0f)
 		[
-			SNew(SSimpleGradient)
-			.StartColor(FLinearColor(161.0f / 255.0f, 57.0f / 255.0f, 191.0f / 255.0f))
-			.EndColor(FLinearColor(100.0f / 255.0f, 100.0f / 255.0f, 100.0f / 255.0f))
+			SNew(SVerticalBox)
+			+SVerticalBox::Slot()
+			.FillHeight(0.15f)
+			[
+				SNew(SVerticalBox)
+				+SVerticalBox::Slot()
+				.VAlign(VAlign_Top)
+				.AutoHeight()
+				.Padding(0.0f, 10.0f)
+				[
+					SNew(SHeader)
+					[
+						SNew(STextBlock)
+						.Text(LOCTEXT("General", "General"))
+					]
+				]
+			]
+			+SVerticalBox::Slot()
+			.FillHeight(0.6f)
+			[
+				SNew(SVerticalBox)
+				+SVerticalBox::Slot()
+				.VAlign(VAlign_Top)
+				.AutoHeight()
+				.Padding(0.0f, 10.0f)
+				[
+					SNew(SHeader)
+					[
+						SNew(STextBlock)
+						.Text(LOCTEXT("Categories", "Categories"))
+					]
+				]
+			]
+			+SVerticalBox::Slot()
+			.FillHeight(0.35f)
+			[
+				SNew(SVerticalBox)
+				+SVerticalBox::Slot()
+				.VAlign(VAlign_Top)
+				.AutoHeight()
+				.Padding(0.0f, 10.0f)
+				[
+					SNew(SHorizontalBox)
+					+SHorizontalBox::Slot()
+					[
+						SNew(SHeader)
+						[
+							SNew(STextBlock)
+							.Text(LOCTEXT("CustomView", "Custom View"))
+						]
+					]
+					+SHorizontalBox::Slot()
+					.HAlign(HAlign_Right)
+					.VAlign(VAlign_Center)
+					.AutoWidth()
+					.Padding(20.0f, 0.0f, 0.0f, 0.0f)
+					[
+						SNew(SButton)
+						.Text(LOCTEXT("CustomViewSyntax", "Syntax"))
+						.OnClicked(this, &SSyncFilterWindow::OnCustomViewSyntaxClicked)
+					]
+				]
+				+SVerticalBox::Slot()
+				.VAlign(VAlign_Fill)
+				[
+					SNew(SMultiLineEditableTextBox)
+					.Padding(10.0f)
+					.AutoWrapText(true)
+					.AlwaysShowScrollbars(true)
+					.HScrollBar(InvisibleHorizontalScrollbar)
+					.BackgroundColor(FLinearColor::Transparent)
+					.Justification(ETextJustify::Left)
+				]
+			]
 		]
 		// Buttons
 		+SVerticalBox::Slot()
@@ -81,7 +158,7 @@ FReply SSyncFilterWindow::OnShowCombinedFilterClicked()
 	fprintf(stderr, "Sync filters:\n%sEnd of sync filters\n", TCHAR_TO_ANSI(*FString::Join(Tab->GetSyncFilters(), TEXT("\n"))));
 
 	TSharedRef<SPopupTextWindow> CombinedFilterWindow = SNew(SPopupTextWindow)
-		.TitleText(FText::FromString("Combined Sync Filter"))
+		.TitleText(LOCTEXT("CombinedSyncFilterWindowTitle", "Combined Sync Filter"))
 		.BodyText(FText::FromString(FString::Join(Tab->GetCombinedSyncFilter(), TEXT("\n"))))
 		.BodyTextJustification(ETextJustify::Left)
 		.ShowScrollBars(true);
@@ -89,6 +166,23 @@ FReply SSyncFilterWindow::OnShowCombinedFilterClicked()
 	return FReply::Handled();
 }
 
+FReply SSyncFilterWindow::OnCustomViewSyntaxClicked()
+{
+	FText Body = FText::FromString(
+		"Specify a custom view of the stream using Perforce-style wildcards, one per line.\n\n"
+		"  - All files are visible by default.\n"
+		"  - To exclude files matching a pattern, prefix it with a '-' character (eg. -/Engine/Documentation/...)\n"
+		"  - Patterns may match any file fragment (eg. *.pdb), or may be rooted to the branch (eg. /Engine/Binaries/.../*.pdb).\n\n"
+		"The view for the current workspace will be appended to the view shared by all workspaces."
+	);
+	TSharedRef<SPopupTextWindow> CombinedFilterWindow = SNew(SPopupTextWindow)
+		.TitleText(LOCTEXT("CustomSyncFilterSyntaxWindow", "Custom Sync Filter Syntax"))
+		.BodyText(Body)
+		.BodyTextJustification(ETextJustify::Left);
+	FSlateApplication::Get().AddModalWindow(CombinedFilterWindow, SharedThis(this), false);
+	return FReply::Handled();
+}
+
 FReply SSyncFilterWindow::OnOkClicked()
 {
 	// Todo: save sync filters
diff --git a/Engine/Source/Programs/SlateUGS/Private/Widgets/SSyncFilterWindow.h b/Engine/Source/Programs/SlateUGS/Private/Widgets/SSyncFilterWindow.h
--- a/Engine/Source/Programs/SlateUGS/Private/Widgets/SSyncFilterWindow.h
+++ b/Engine/Source/Programs/SlateUGS/Private/Widgets/SSyncFilterWindow.h
@@ -17,6 +17,7 @@ public:
 
 private:
 	FReply OnShowCombinedFilterClicked();
+	FReply OnCustomViewSyntaxClicked();
 	FReply OnOkClicked();
 	FReply OnCancelClicked();
 
