Fix compiler error

diff --git a/Engine/Source/Programs/SlateUGS/Private/Widgets/SWorkspaceWindow.cpp b/Engine/Source/Programs/SlateUGS/Private/Widgets/SWorkspaceWindow.cpp
--- a/Engine/Source/Programs/SlateUGS/Private/Widgets/SWorkspaceWindow.cpp
+++ b/Engine/Source/Programs/SlateUGS/Private/Widgets/SWorkspaceWindow.cpp
@@ -215,9 +215,9 @@ FReply SWorkspaceWindow::OnOkClicked()
 	else
 	{
 		// Todo: report other errors (no workspace associated with the file)
-		FText Title = FText::FromString("Error");
+		FText WorkspaceTitle = FText::FromString("Error");
 		FText Error = FText::FromString("Project file does not exist, try again");
-		FSlateApplication::Get().AddModalWindow(SNew(SPopupTextWindow).TitleText(Title).BodyText(Error), SharedThis(this), false);
+		FSlateApplication::Get().AddModalWindow(SNew(SPopupTextWindow).TitleText(WorkspaceTitle).BodyText(Error), SharedThis(this), false);
 	}
 
 	return FReply::Handled();
