Fix p4.exe spawning a cmd prompt

diff --git a/Engine/Source/Programs/SlateUGS/Private/UGSCore/Perforce.cpp b/Engine/Source/Programs/SlateUGS/Private/UGSCore/Perforce.cpp
--- a/Engine/Source/Programs/SlateUGS/Private/UGSCore/Perforce.cpp
+++ b/Engine/Source/Programs/SlateUGS/Private/UGSCore/Perforce.cpp
@@ -52,7 +52,7 @@ struct FPerforceExe
 		void* WritePipe = nullptr;
 		FPlatformProcess::CreatePipe(ReadPipe, WritePipe);
 
-		FProcHandle P4Proc = FPlatformProcess::CreateProc(*PerforceExe, *CommandLine, false, false, false, &ProcId, 0, nullptr, WritePipe, ReadPipe);
+		FProcHandle P4Proc = FPlatformProcess::CreateProc(*PerforceExe, *CommandLine, false, true, true, &ProcId, 0, nullptr, WritePipe, ReadPipe);
 
 		FString P4Output;
 		FString LatestOutput = FPlatformProcess::ReadPipe(ReadPipe);
@@ -1364,7 +1364,7 @@ bool FPerforceConnection::RunCommandWithBinaryOutput(const FString& CommandLine,
     void* WritePipe = nullptr;
     FPlatformProcess::CreatePipe(ReadPipe, WritePipe);
 
-	FProcHandle P4Proc = FPlatformProcess::CreateProc(*GPerforceExe.PerforceExe, *FullCommandLine, false, false, false, &ProcId, 0, nullptr, WritePipe, ReadPipe);
+	FProcHandle P4Proc = FPlatformProcess::CreateProc(*GPerforceExe.PerforceExe, *FullCommandLine, false, true, true, &ProcId, 0, nullptr, WritePipe, ReadPipe);
 
 	size_t BufferPos = 0;
 	size_t BufferEnd = 0;
