Remove todo that was fixed

diff --git a/Engine/Source/Programs/SlateUGS/Private/UGSCore/Perforce.cpp b/Engine/Source/Programs/SlateUGS/Private/UGSCore/Perforce.cpp
--- a/Engine/Source/Programs/SlateUGS/Private/UGSCore/Perforce.cpp
+++ b/Engine/Source/Programs/SlateUGS/Private/UGSCore/Perforce.cpp
@@ -65,7 +65,6 @@ struct FPerforceExe
 
 		FPlatformProcess::ClosePipe(ReadPipe, WritePipe);
 
-		// TODO check if p4 spits out \n on Windows
 		OutLines = Split(P4Output, LINE_TERMINATOR);
 
 		int ExitCode = -1;
