Created the list of changes

diff --git a/Engine/Source/Programs/SlateUGS/Private/ChangeInfo.h b/Engine/Source/Programs/SlateUGS/Private/ChangeInfo.h
new file mode 100644
--- /dev/null
+++ b/Engine/Source/Programs/SlateUGS/Private/ChangeInfo.h
@@ -0,0 +1,15 @@
+// Copyright Epic Games, Inc. All Rights Reserved.
+
+#pragma once
+
+#include "UGSCore/EventMonitor.h"
+
+struct FChangeInfo
+{
+	EReviewVerdict ReviewStatus;
+	FText Changelist;
+	FText Time;
+	FText Author;
+	FText Description;
+	// Todo: add Horde badges
+};
diff --git a/Engine/Source/Programs/SlateUGS/Private/HordeBuildRowInfo.h b/Engine/Source/Programs/SlateUGS/Private/HordeBuildRowInfo.h
deleted file mode 100644
diff --git a/Engine/Source/Programs/SlateUGS/Private/UGSTab.cpp b/Engine/Source/Programs/SlateUGS/Private/UGSTab.cpp
--- a/Engine/Source/Programs/SlateUGS/Private/UGSTab.cpp
+++ b/Engine/Source/Programs/SlateUGS/Private/UGSTab.cpp
@@ -3,6 +3,8 @@
 #include "UGSTab.h"
 
 #include "UGSLog.h"
+#include "ChangeInfo.h"
+
 #include "UGSCore/Utility.h"
 #include "UGSCore/BuildStep.h"
 #include "UGSCore/DetectProjectSettingsTask.h"
@@ -201,6 +203,27 @@ namespace
 		Variables.Add("UE4EditorDebugArg", (EditorBuildConfig == EBuildConfig::Debug || EditorBuildConfig == EBuildConfig::DebugGame)? " -debug" : "");
 		return Variables;
 	}
+
+	static FString FormatUserName(FString UserName)
+	{
+		TStringBuilder<256> NormalUserName;
+		for(int Index = 0; Index < UserName.Len(); Index++)
+		{
+			if(Index == 0 || UserName[Index - 1] == '.')
+			{
+				NormalUserName.AppendChar(FChar::ToUpper(UserName[Index]));
+			}
+			else if(UserName[Index] == '.')
+			{
+				NormalUserName.AppendChar(' ');
+			}
+			else
+			{
+				NormalUserName.AppendChar(FChar::ToLower(UserName[Index]));
+			}
+		}
+		return NormalUserName.ToString();
+	}
 }
 
 bool UGSTab::OnWorkspaceChosen(const FString& Project)
@@ -321,6 +344,7 @@ bool UGSTab::ShouldIncludeInReviewedList(const TSet<int>& PromotedChangeNumbers,
 void UGSTab::UpdateGameTabBuildList()
 {
 	TArray<TSharedRef<FPerforceChangeSummary, ESPMode::ThreadSafe>> Changes = PerforceMonitor->GetChanges();
+	TArray<TSharedPtr<FChangeInfo>> ChangeInfos;
 
 	// do we need to store these for something?
 	TSet<int> PromotedChangeNumbers = PerforceMonitor->GetPromotedChangeNumbers();
@@ -351,15 +375,36 @@ void UGSTab::UpdateGameTabBuildList()
 					DisplayTime = (DisplayTime - DetectSettings->ServerTimeZone).GetDate();
 				}
 
+				TSharedPtr<FEventSummary> Review;
+				EReviewVerdict Status = EReviewVerdict::Unknown;
+				if (EventMonitor->TryGetSummaryForChange(Change.Number, Review) && Review->LastStarReview.IsValid())
+				{
+					if (Review->LastStarReview->Type == EEventType::Starred)
+					{
+						Status = EReviewVerdict::Good;
+					}
+					else
+					{
+						Status = Review->Verdict;
+					}
+				}
+
 				TArray<FString> DescLines;
 
 				// split on \n so we can just take the first line
 				Change.Description.ParseIntoArray(DescLines, TEXT("\n"));
 
-				printf("%s CL %i Author: %s@%s :: %s\n", TCHAR_TO_ANSI(*DisplayTime.ToString()), Change.Number, TCHAR_TO_ANSI(*Change.User), TCHAR_TO_ANSI(*Change.Client), TCHAR_TO_ANSI(*DescLines[0]));
+				ChangeInfos.Add(MakeShareable(new FChangeInfo{
+					Status,
+					FText::FromString(FString::FromInt(Change.Number)),
+					FText::FromString(DisplayTime.ToString()),
+					FText::FromString(FormatUserName(Change.User)),
+					FText::FromString(DescLines[0])}));
 			}
 		}
 	}
+
+	GameSyncTabView->AddHordeBuilds(ChangeInfos);
 }
 
 bool UGSTab::IsSyncing() const
diff --git a/Engine/Source/Programs/SlateUGS/Private/Widgets/SGameSyncTab.cpp b/Engine/Source/Programs/SlateUGS/Private/Widgets/SGameSyncTab.cpp
--- a/Engine/Source/Programs/SlateUGS/Private/Widgets/SGameSyncTab.cpp
+++ b/Engine/Source/Programs/SlateUGS/Private/Widgets/SGameSyncTab.cpp
@@ -19,96 +19,87 @@
 
 #include "UGSTab.h"
 
+#include "Math/RandomStream.h" // Todo: Delete
+
 #define LOCTEXT_NAMESPACE "UGSWindow"
 
-TSharedRef<ITableRow> SGameSyncTab::GenerateHordeBuildTableRow(TSharedPtr<HordeBuildRowInfo> InItem, const TSharedRef<STableViewBase>& InOwnerTable)
+namespace
 {
-	return SNew(STableRow<TSharedPtr<HordeBuildRowInfo>>, InOwnerTable) // Todo: Maybe replace with SMultiColumnTableRow
-	[
-		SNew(SHorizontalBox)
-		// Build status
-		+SHorizontalBox::Slot()
-		.AutoWidth()
-		.Padding(HordeBuildRowHorizontalPadding + HordeBuildRowExtraIconPadding, HordeBuildRowVerticalPadding)
-		[
-			SNew(SImage)
-			.Image(FAppStyle::Get().GetBrush("Icons.FilledCircle"))
-			.ColorAndOpacity(InItem->bBuildStatus
-				? FLinearColor(116.0f / 255.0f, 160.0f / 255.0f, 64.0f / 255.0f) // Todo: store literals
-				: FLinearColor(209.0f / 255.0f, 56.0f / 255.0f, 56.0f / 255.0f))
-		]
-		// CODE/CONTENT
-		+SHorizontalBox::Slot()
-		.AutoWidth()
-		.Padding(HordeBuildRowHorizontalPadding, HordeBuildRowVerticalPadding)
-		[
-			SNew(STextBlock)
-			.Text(LOCTEXT("CodeContent", "CODE / CONTENT"))
-		]
-		// Changelist
-		+SHorizontalBox::Slot()
-		.AutoWidth()
-		.Padding(HordeBuildRowHorizontalPadding, HordeBuildRowVerticalPadding)
-		[
-			SNew(STextBlock)
-			.Text(InItem->Changelist)
-		]
-		// Time
-		+SHorizontalBox::Slot()
-		.AutoWidth()
-		.Padding(HordeBuildRowHorizontalPadding, HordeBuildRowVerticalPadding)
-		[
-			SNew(STextBlock)
-			.Text(InItem->Changelist)
-		]
-		// Author
-		+SHorizontalBox::Slot()
-		.AutoWidth()
-		.Padding(HordeBuildRowHorizontalPadding, HordeBuildRowVerticalPadding)
-		[
-			SNew(STextBlock)
-			.Text(InItem->Author)
-		]
-		// Description
-		+SHorizontalBox::Slot()
-		.AutoWidth()
-		.Padding(HordeBuildRowHorizontalPadding, HordeBuildRowVerticalPadding)
-		[
-			SNew(STextBlock)
-			.Text(InItem->Description)
-		]
-		// EDITOR
-		+SHorizontalBox::Slot()
-		.AutoWidth()
-		.Padding(HordeBuildRowHorizontalPadding, HordeBuildRowVerticalPadding)
-		[
-			SNew(STextBlock)
-			.Text(LOCTEXT("Editor", "EDITOR"))
-		]
-		// PLATFORMS
-		+SHorizontalBox::Slot()
-		.AutoWidth()
-		.Padding(HordeBuildRowHorizontalPadding, HordeBuildRowVerticalPadding)
-		[
-			SNew(STextBlock)
-			.Text(LOCTEXT("Platforms", "PLATFORMS"))
-		]
-		// CIS
-		+SHorizontalBox::Slot()
-		.AutoWidth()
-		.Padding(HordeBuildRowHorizontalPadding, HordeBuildRowVerticalPadding)
-		[
-			SNew(STextBlock)
-			.Text(LOCTEXT("CIS", "CIS"))
-		]
-		// Status
-		+SHorizontalBox::Slot()
-		.Padding(HordeBuildRowHorizontalPadding, HordeBuildRowVerticalPadding)
+	const FName HordeTableColumnStatus(TEXT("Status"));
+	const FName HordeTableColumnChange(TEXT("Change"));
+	const FName HordeTableColumnTime(TEXT("Time"));
+	const FName HordeTableColumnAuthor(TEXT("Author"));
+	const FName HordeTableColumnDescription(TEXT("Description"));
+}
+
+void SBuildDataRow::Construct(const FArguments& InArgs, const TSharedRef<STableViewBase>& InOwnerTableView, const TSharedPtr<FChangeInfo>& Item)
+{
+	CurrentItem = Item;
+	SMultiColumnTableRow<TSharedPtr<FChangeInfo>>::Construct(SMultiColumnTableRow::FArguments(), InOwnerTableView);
+}
+
+TSharedRef<SWidget> SBuildDataRow::GenerateWidgetForColumn(const FName& ColumnId)
+{
+	if (ColumnId == HordeTableColumnStatus)
+	{
+		TSharedRef<SImage> StatusCircle = SNew(SImage)
+			.Image(FAppStyle::Get().GetBrush("Icons.FilledCircle"));
+
+		switch (CurrentItem->ReviewStatus)
+		{
+			case EReviewVerdict::Good:
+				StatusCircle->SetColorAndOpacity(FLinearColor::Green);
+				break;
+			case EReviewVerdict::Bad:
+				StatusCircle->SetColorAndOpacity(FLinearColor::Red);
+				break;
+			case EReviewVerdict::Mixed:
+				StatusCircle->SetColorAndOpacity(FLinearColor::Yellow);
+				break;
+			case EReviewVerdict::Unknown:
+			default:
+				StatusCircle->SetColorAndOpacity(FLinearColor::Gray);
+				break;
+		}
+
+		return SNew(SBox)
+			.HAlign(HAlign_Center)
+			.VAlign(VAlign_Center)
+			[
+				StatusCircle
+			];
+	}
+
+	TSharedRef<STextBlock> TextItem = SNew(STextBlock);
+	if (ColumnId == HordeTableColumnChange)
+	{
+		TextItem->SetText(CurrentItem->Changelist);
+		TextItem->SetJustification(ETextJustify::Center);
+	}
+	if (ColumnId == HordeTableColumnTime)
+	{
+		TextItem->SetText(CurrentItem->Time);
+		TextItem->SetJustification(ETextJustify::Center);
+	}
+	if (ColumnId == HordeTableColumnAuthor)
+	{
+		TextItem->SetText(CurrentItem->Author);
+	}
+	if (ColumnId == HordeTableColumnDescription)
+	{
+		TextItem->SetText(CurrentItem->Description);
+	}
+
+	return SNew(SBox)
+		.Padding(10.0f, 0.0f)
 		[
-			SNew(STextBlock)
-			.Text(InItem->Status)
-		]
-	];
+			TextItem
+		];
+}
+
+TSharedRef<ITableRow> SGameSyncTab::GenerateHordeBuildTableRow(TSharedPtr<FChangeInfo> InItem, const TSharedRef<STableViewBase>& InOwnerTable)
+{
+	return SNew(SBuildDataRow, InOwnerTable, InItem);
 }
 
 // Button callbacks
@@ -135,7 +126,6 @@ TSharedRef<SWidget> SGameSyncTab::MakeSyncButtonDropdown()
 void SGameSyncTab::Construct(const FArguments& InArgs)
 {
 	Tab = InArgs._Tab;
-	HordeBuilds = InArgs._HordeBuilds;
 
 	this->ChildSlot
 	[
@@ -218,7 +208,6 @@ void SGameSyncTab::Construct(const FArguments& InArgs)
 					[
 						SNew(SSimpleButton)
 						.Icon(FAppStyle::Get().GetBrush("GraphEditor.Clean")) // Todo: shouldn't use this icon (repurposing)
-						.Text(LOCTEXT("CleanSolution", "Clean Solution"))
 					]
 					+SHorizontalBox::Slot()
 					.AutoWidth()
@@ -355,29 +344,27 @@ void SGameSyncTab::Construct(const FArguments& InArgs)
 		.Padding(20.0f, 5.0f)
 		.FillHeight(0.45f)
 		[
-			SAssignNew(HordeBuildsView, SListView<TSharedPtr<HordeBuildRowInfo>>)
+			SAssignNew(HordeBuildsView, SListView<TSharedPtr<FChangeInfo>>)
 			.IsEnabled_Lambda([this] { return !Tab->IsSyncing(); })
+			.OnGenerateRow(this, &SGameSyncTab::GenerateHordeBuildTableRow)
 			.ListItemsSource(&HordeBuilds)
 			.HeaderRow(
 				SNew(SHeaderRow)
-				+SHeaderRow::Column(FName(TEXT("TYPE")))
-				.DefaultLabel(LOCTEXT("HordeHeaderType", "TYPE"))
-				+SHeaderRow::Column(FName(TEXT("CHANGE")))
-				.DefaultLabel(LOCTEXT("HordeHeaderChange", "CHANGE"))
-				+SHeaderRow::Column(FName(TEXT("AUTHOR")))
-				.DefaultLabel(LOCTEXT("HordeHeaderAuthor", "AUTHOR"))
-				+SHeaderRow::Column(FName(TEXT("DESCRIPTION")))
-				.DefaultLabel(LOCTEXT("HordeHeaderDescription", "DESCRIPTION"))
-				+SHeaderRow::Column(FName(TEXT("EDITOR")))
-				.DefaultLabel(LOCTEXT("HordeHeaderEditor", "EDITOR"))
-				+SHeaderRow::Column(FName(TEXT("PLATFORMS")))
-				.DefaultLabel(LOCTEXT("HordeHeaderPlatforms", "PLATFORMS"))
-				+SHeaderRow::Column(FName(TEXT("CIS")))
-				.DefaultLabel(LOCTEXT("HordeHeaderCIS", "CIS"))
-				+SHeaderRow::Column(FName(TEXT("STATUS")))
-				.DefaultLabel(LOCTEXT("HordeHeaderStatus", "STATUS"))
+				+SHeaderRow::Column(HordeTableColumnStatus)
+				.DefaultLabel(LOCTEXT("HordeHeaderStatus", ""))
+				.FixedWidth(35.0f)
+				+SHeaderRow::Column(HordeTableColumnChange)
+				.DefaultLabel(LOCTEXT("HordeHeaderChange", "Change"))
+				.FixedWidth(100.0f)
+				+SHeaderRow::Column(HordeTableColumnTime)
+				.DefaultLabel(LOCTEXT("HordeHeaderTime", "Time"))
+				.FillWidth(0.2f)
+				+SHeaderRow::Column(HordeTableColumnAuthor)
+				.DefaultLabel(LOCTEXT("HordeHeaderAuthor", "Author"))
+				.FillWidth(0.15f)
+				+SHeaderRow::Column(HordeTableColumnDescription)
+				.DefaultLabel(LOCTEXT("HordeHeaderDescription", "Description"))
 			)
-			.OnGenerateRow(this, &SGameSyncTab::GenerateHordeBuildTableRow)
 		]
 		// Log
 		+SVerticalBox::Slot()
@@ -417,13 +404,21 @@ void SGameSyncTab::SetStreamPathText(FText StreamPath)
 {
 	StreamPathText->SetText(StreamPath);
 }
+
 void SGameSyncTab::SetChangelistText(FText Changelist)
 {
 	ChangelistText->SetText(Changelist);
 }
+
 void SGameSyncTab::SetProjectPathText(FText ProjectPath)
 {
 	ProjectPathText->SetText(ProjectPath);
 }
 
+void SGameSyncTab::AddHordeBuilds(const TArray<TSharedPtr<FChangeInfo>>& Builds)
+{
+	HordeBuilds += Builds;
+	HordeBuildsView->RebuildList();
+}
+
 #undef LOCTEXT_NAMESPACE
diff --git a/Engine/Source/Programs/SlateUGS/Private/Widgets/SGameSyncTab.h b/Engine/Source/Programs/SlateUGS/Private/Widgets/SGameSyncTab.h
--- a/Engine/Source/Programs/SlateUGS/Private/Widgets/SGameSyncTab.h
+++ b/Engine/Source/Programs/SlateUGS/Private/Widgets/SGameSyncTab.h
@@ -3,16 +3,27 @@
 #pragma once
 
 #include "Widgets/SCompoundWidget.h"
-#include "HordeBuildRowInfo.h"
+#include "ChangeInfo.h"
 
 class UGSTab;
 class SLogWidget;
 
-class SGameSyncTab : public SCompoundWidget
+class SBuildDataRow final : public SMultiColumnTableRow<TSharedPtr<FChangeInfo>>
+{
+	SLATE_BEGIN_ARGS(SBuildDataRow) {}
+	SLATE_END_ARGS()
+
+	void Construct(const FArguments& InArgs, const TSharedRef<STableViewBase>& InOwnerTableView, const TSharedPtr<FChangeInfo>& Item);
+	virtual TSharedRef<SWidget> GenerateWidgetForColumn(const FName& ColumnId) override;
+
+protected:
+	TSharedPtr<FChangeInfo> CurrentItem;
+};
+
+class SGameSyncTab final : public SCompoundWidget
 {
 public:
 	SLATE_BEGIN_ARGS(SGameSyncTab) {}
-		SLATE_ARGUMENT(TArray<TSharedPtr<HordeBuildRowInfo>>, HordeBuilds) // Todo: remove once we start getting real horde data
 		SLATE_ARGUMENT(UGSTab*, Tab)
 	SLATE_END_ARGS()
 
@@ -27,12 +38,14 @@ public:
 	void SetChangelistText(FText Changelist);
 	void SetProjectPathText(FText ProjectPath);
 
+	void AddHordeBuilds(const TArray<TSharedPtr<FChangeInfo>>& Builds);
+
 private:
-	TSharedRef<ITableRow> GenerateHordeBuildTableRow(TSharedPtr<HordeBuildRowInfo> InItem, const TSharedRef<STableViewBase>& InOwnerTable);
+	TSharedRef<ITableRow> GenerateHordeBuildTableRow(TSharedPtr<FChangeInfo> InItem, const TSharedRef<STableViewBase>& InOwnerTable);
 	TSharedRef<SWidget> MakeSyncButtonDropdown();
 
-	TSharedPtr<SListView<TSharedPtr<HordeBuildRowInfo>>> HordeBuildsView;
-	TArray<TSharedPtr<HordeBuildRowInfo>> HordeBuilds;
+	TSharedPtr<SListView<TSharedPtr<FChangeInfo>>> HordeBuildsView;
+	TArray<TSharedPtr<FChangeInfo>> HordeBuilds;
 
 	TSharedPtr<SLogWidget> SyncLog;
 
diff --git a/Engine/Source/Programs/SlateUGS/Private/Widgets/SWorkspaceWindow.cpp b/Engine/Source/Programs/SlateUGS/Private/Widgets/SWorkspaceWindow.cpp
--- a/Engine/Source/Programs/SlateUGS/Private/Widgets/SWorkspaceWindow.cpp
+++ b/Engine/Source/Programs/SlateUGS/Private/Widgets/SWorkspaceWindow.cpp
@@ -6,7 +6,6 @@
 #include "DesktopPlatformModule.h"
 
 #include "UGSTab.h"
-#include "HordeBuildRowInfo.h"
 #include "SGameSyncTab.h"
 #include "SPopupTextWindow.h"
 
