Fix Build.sh hardcoded for Linux only

diff --git a/Engine/Source/Programs/SlateUGS/Private/UGSCore/Workspace.cpp b/Engine/Source/Programs/SlateUGS/Private/UGSCore/Workspace.cpp
--- a/Engine/Source/Programs/SlateUGS/Private/UGSCore/Workspace.cpp
+++ b/Engine/Source/Programs/SlateUGS/Private/UGSCore/Workspace.cpp
@@ -850,8 +850,13 @@ EWorkspaceUpdateResult FWorkspace::UpdateWorkspaceInternal(FWorkspaceUpdateConte
 		Progress.Set(TEXT("Starting build..."), 0.0f);
 
 		// Check we've built UBT (it should have been compiled by generating project files)
-		//FString UnrealBuildToolPath = LocalRootPath / TEXT("Engine/Binaries/DotNET/UnrealBuildTool.exe");
+#if PLATFORM_WINDOWS
+		FString UnrealBuildToolPath = LocalRootPath / TEXT("Engine/Build/BatchFiles/Build.bat");
+#elif PLATFORM_MAC
+		FString UnrealBuildToolPath = LocalRootPath / TEXT("Engine/Build/BatchFiles/Mac/Build.sh");
+#else
 		FString UnrealBuildToolPath = LocalRootPath / TEXT("Engine/Build/BatchFiles/Linux/Build.sh");
+#endif
 		if(!IFileManager::Get().FileExists(*UnrealBuildToolPath))
 		{
 			OutStatusMessage = FString::Printf(TEXT("Couldn't find %s"), *UnrealBuildToolPath);
