Sync filters completed

diff --git a/Engine/Source/Programs/SlateUGS/Private/UGSTab.cpp b/Engine/Source/Programs/SlateUGS/Private/UGSTab.cpp
--- a/Engine/Source/Programs/SlateUGS/Private/UGSTab.cpp
+++ b/Engine/Source/Programs/SlateUGS/Private/UGSTab.cpp
@@ -269,6 +269,21 @@ void UGSTab::OnSyncLatest()
 	FPlatformProcess::ReturnSynchEventToPool(AbortEvent);
 }
 
+void UGSTab::OnSyncFilterWindowSaved(
+	const TArray<FString>& SyncViewCurrent,
+	const TArray<FGuid>& SyncExcludedCategoriesCurrent,
+	const TArray<FString>& SyncViewAll,
+	const TArray<FGuid>& SyncExcludedCategoriesAll)
+{
+	WorkspaceSettings->SyncView = SyncViewCurrent;
+	WorkspaceSettings->SyncExcludedCategories = SyncExcludedCategoriesCurrent;
+
+	UserSettings->SyncView = SyncViewAll;
+	UserSettings->SyncExcludedCategories = SyncExcludedCategoriesAll;
+
+	UserSettings->Save();
+}
+
 void UGSTab::QueueMessageForMainThread(TFunction<void()> Function)
 {
 	FScopeLock Lock(&CriticalSection);
@@ -417,9 +432,36 @@ FString UGSTab::GetSyncProgress() const
 	return Workspace->GetCurrentProgress().Key;
 }
 
-const TArray<FString>& UGSTab::GetSyncFilters() const
+TArray<FWorkspaceSyncCategory> UGSTab::GetSyncCategories(SyncCategoryType CategoryType) const
 {
-	return WorkspaceSettings->SyncView;
+	TArray<FWorkspaceSyncCategory> SyncCategories;
+	TMap<FGuid, FWorkspaceSyncCategory> Categories = Workspace->GetSyncCategories();
+	TArray<FGuid> SyncExcludedCategories;
+
+	if (CategoryType == SyncCategoryType::CurrentWorkspace)
+	{
+		SyncExcludedCategories = WorkspaceSettings->SyncExcludedCategories;
+	}
+	else
+	{
+		SyncExcludedCategories = UserSettings->SyncExcludedCategories;
+	}
+
+	for (const FGuid& Guid : SyncExcludedCategories)
+	{
+		Categories[Guid].bEnable = false;
+	}
+	Categories.GenerateValueArray(SyncCategories);
+	return SyncCategories;
+}
+
+TArray<FString> UGSTab::GetSyncViews(SyncCategoryType CategoryType) const
+{
+	if (CategoryType == SyncCategoryType::CurrentWorkspace)
+	{
+		return WorkspaceSettings->SyncView;
+	}
+	return UserSettings->SyncView;
 }
 
 const TArray<FString>& UGSTab::GetCombinedSyncFilter() const
diff --git a/Engine/Source/Programs/SlateUGS/Private/UGSTab.h b/Engine/Source/Programs/SlateUGS/Private/UGSTab.h
--- a/Engine/Source/Programs/SlateUGS/Private/UGSTab.h
+++ b/Engine/Source/Programs/SlateUGS/Private/UGSTab.h
@@ -15,6 +15,12 @@ class FDetectProjectSettingsTask;
 class FPerforceMonitor;
 class FEventMonitor;
 
+enum SyncCategoryType
+{
+	CurrentWorkspace,
+	AllWorkspaces
+};
+
 class UGSTab
 {
 public:
@@ -30,10 +36,19 @@ public:
 	// Slate callbacks
 	bool OnWorkspaceChosen(const FString& Project);
 	void OnSyncLatest();
+	void OnSyncFilterWindowSaved(
+		const TArray<FString>& SyncViewCurrent,
+		const TArray<FGuid>& SyncExcludedCategoriesCurrent,
+		const TArray<FString>& SyncViewAll,
+		const TArray<FGuid>& SyncExcludedCategoriesAll);
+
+	// Accessors
 	bool IsSyncing() const;
 	FString GetSyncProgress() const;
-	const TArray<FString>& GetSyncFilters() const; // Todo: is return type okay?
-	const TArray<FString>& GetCombinedSyncFilter() const; // Todo: is return type okay?
+	const TArray<FString>& GetSyncFilters() const;
+	const TArray<FString>& GetCombinedSyncFilter() const;
+	TArray<FWorkspaceSyncCategory> GetSyncCategories(SyncCategoryType CategoryType) const;
+	TArray<FString> GetSyncViews(SyncCategoryType CategoryType) const;
 private:
 
 	void OnWorkspaceSyncComplete(
diff --git a/Engine/Source/Programs/SlateUGS/Private/Widgets/SGameSyncTab.cpp b/Engine/Source/Programs/SlateUGS/Private/Widgets/SGameSyncTab.cpp
--- a/Engine/Source/Programs/SlateUGS/Private/Widgets/SGameSyncTab.cpp
+++ b/Engine/Source/Programs/SlateUGS/Private/Widgets/SGameSyncTab.cpp
@@ -200,13 +200,14 @@ void SGameSyncTab::Construct(const FArguments& InArgs)
 					.AutoWidth()
 					[
 						SNew(SSimpleButton)
-						.Text(LOCTEXT("OpenPerforce", "OpenPerforce"))
+						.Text(LOCTEXT("OpenPerforce", "Open Perforce"))
 						.Icon(FAppStyle::Get().GetBrush("Icons.Blueprints")) // Todo: shouldn't use this icon (repurposing)
 					]
 					+SHorizontalBox::Slot()
 					.AutoWidth()
 					[
 						SNew(SSimpleButton)
+						.Text(LOCTEXT("CleanSolution", "Clean Solution"))
 						.Icon(FAppStyle::Get().GetBrush("GraphEditor.Clean")) // Todo: shouldn't use this icon (repurposing)
 					]
 					+SHorizontalBox::Slot()
diff --git a/Engine/Source/Programs/SlateUGS/Private/Widgets/SSyncFilterWindow.cpp b/Engine/Source/Programs/SlateUGS/Private/Widgets/SSyncFilterWindow.cpp
--- a/Engine/Source/Programs/SlateUGS/Private/Widgets/SSyncFilterWindow.cpp
+++ b/Engine/Source/Programs/SlateUGS/Private/Widgets/SSyncFilterWindow.cpp
@@ -4,7 +4,6 @@
 
 #include "UGSTab.h"
 
-#include "Widgets/Input/SMultiLineEditableTextBox.h"
 #include "Widgets/Layout/SUniformGridPanel.h"
 #include "Widgets/Layout/SHeader.h"
 #include "SPrimaryButton.h"
@@ -13,14 +12,71 @@
 
 #define LOCTEXT_NAMESPACE "SSyncFilterWindow"
 
-void SSyncFilterWindow::Construct(const FArguments& InArgs)
+void SSyncFilterWindow::ConstructSyncFilters()
 {
-	Tab = InArgs._Tab;
+	WorkspaceCategoriesCurrent = Tab->GetSyncCategories(SyncCategoryType::CurrentWorkspace);
+	SyncFiltersCurrent = SNew(SCheckBoxList)
+		.ItemHeaderLabel(LOCTEXT("CheckBoxListCurrent", "Sync Filters for the current workspace"))
+		.IncludeGlobalCheckBoxInHeaderRow(false);
+	for (const FWorkspaceSyncCategory& Category : WorkspaceCategoriesCurrent)
+	{
+		TSharedRef<SBox> CheckBoxWrapper = SNew(SBox)
+			.Padding(5.0f)
+			[
+				SNew(STextBlock)
+				.Text(FText::FromString(Category.Name))
+			];
+		SyncFiltersCurrent->AddItem(CheckBoxWrapper, Category.bEnable);
+	}
+
+	WorkspaceCategoriesAll = Tab->GetSyncCategories(SyncCategoryType::AllWorkspaces);
+	SyncFiltersAll = SNew(SCheckBoxList)
+		.ItemHeaderLabel(LOCTEXT("CheckBoxListAll", "Sync Filters for all workspaces"))
+		.IncludeGlobalCheckBoxInHeaderRow(false);
+	for (const FWorkspaceSyncCategory& Category : WorkspaceCategoriesAll)
+	{
+		TSharedRef<SBox> CheckBoxWrapper = SNew(SBox)
+			.Padding(5.0f)
+			[
+				SNew(STextBlock)
+				.Text(FText::FromString(Category.Name))
+			];
+		SyncFiltersAll->AddItem(CheckBoxWrapper, Category.bEnable);
+	}
+}
 
+void SSyncFilterWindow::ConstructCustomSyncViewTextBoxes()
+{
 	TSharedRef<SScrollBar> InvisibleHorizontalScrollbar = SNew(SScrollBar) // Todo: this code is duplicated in SPopupTextWindow.cpp
 		.AlwaysShowScrollbar(false)										   // Is there a simpler way to make the horizontal scroll bar invisible?
 		.Orientation(Orient_Horizontal);								   // If not, I guess we should factor this out into a tiny widget?
 
+	SAssignNew(CustomSyncViewCurrent, SMultiLineEditableTextBox)
+	.Padding(10.0f)
+	.AutoWrapText(true)
+	.AlwaysShowScrollbars(true)
+	.HScrollBar(InvisibleHorizontalScrollbar)
+	.BackgroundColor(FLinearColor::Transparent)
+	.Justification(ETextJustify::Left)
+	.Text(FText::FromString(FString::Join(Tab->GetSyncViews(SyncCategoryType::CurrentWorkspace), TEXT("\n"))));
+
+	SAssignNew(CustomSyncViewAll, SMultiLineEditableTextBox)
+	.Padding(10.0f)
+	.AutoWrapText(true)
+	.AlwaysShowScrollbars(true)
+	.HScrollBar(InvisibleHorizontalScrollbar)
+	.BackgroundColor(FLinearColor::Transparent)
+	.Justification(ETextJustify::Left)
+	.Text(FText::FromString(FString::Join(Tab->GetSyncViews(SyncCategoryType::AllWorkspaces), TEXT("\n"))));
+}
+
+void SSyncFilterWindow::Construct(const FArguments& InArgs)
+{
+	Tab = InArgs._Tab;
+
+	ConstructSyncFilters();
+	ConstructCustomSyncViewTextBoxes();
+
 	SWindow::Construct(SWindow::FArguments()
 	.Title(LOCTEXT("WindowTitle", "Sync Filters"))
 	.SizingRule(ESizingRule::FixedSize)
@@ -35,11 +91,12 @@ void SSyncFilterWindow::Construct(const FArguments& InArgs)
 			SNew(STextBlock)
 			.Text(LOCTEXT("SyncFilterHintText", "Files synced from Perforce may be filtered by a custom stream view, and list of predefined categories. Settings for the current workspace override defaults for all workspaces."))
 		]
-		// Filter checkbox list
+		// Sync filters
 		+SVerticalBox::Slot()
 		.Padding(20.0f, 0.0f)
 		[
 			SNew(SVerticalBox)
+			// General
 			+SVerticalBox::Slot()
 			.FillHeight(0.15f)
 			[
@@ -56,6 +113,7 @@ void SSyncFilterWindow::Construct(const FArguments& InArgs)
 					]
 				]
 			]
+			// Categories
 			+SVerticalBox::Slot()
 			.FillHeight(0.6f)
 			[
@@ -71,7 +129,21 @@ void SSyncFilterWindow::Construct(const FArguments& InArgs)
 						.Text(LOCTEXT("Categories", "Categories"))
 					]
 				]
+				+SVerticalBox::Slot()
+				.VAlign(VAlign_Fill)
+				[
+					SNew(SHorizontalBox)
+					+SHorizontalBox::Slot()
+					[
+						SyncFiltersCurrent.ToSharedRef()
+					]
+					+SHorizontalBox::Slot()
+					[
+						SyncFiltersAll.ToSharedRef()
+					]
+				]
 			]
+			// Custom View
 			+SVerticalBox::Slot()
 			.FillHeight(0.35f)
 			[
@@ -104,13 +176,17 @@ void SSyncFilterWindow::Construct(const FArguments& InArgs)
 				+SVerticalBox::Slot()
 				.VAlign(VAlign_Fill)
 				[
-					SNew(SMultiLineEditableTextBox)
-					.Padding(10.0f)
-					.AutoWrapText(true)
-					.AlwaysShowScrollbars(true)
-					.HScrollBar(InvisibleHorizontalScrollbar)
-					.BackgroundColor(FLinearColor::Transparent)
-					.Justification(ETextJustify::Left)
+					SNew(SHorizontalBox)
+					+SHorizontalBox::Slot()
+					.Padding(0.0f, 0.0f, 10.0f, 0.0f)
+					[
+						CustomSyncViewCurrent.ToSharedRef()
+					]
+					+SHorizontalBox::Slot()
+					.Padding(10.0f, 0.0f, 0.0f, 0.0f)
+					[
+						CustomSyncViewAll.ToSharedRef()
+					]
 				]
 			]
 		]
@@ -138,8 +214,8 @@ void SSyncFilterWindow::Construct(const FArguments& InArgs)
 				+SUniformGridPanel::Slot(0, 0)
 				[
 					SNew(SPrimaryButton)
-					.Text(LOCTEXT("OkButtonText", "Ok"))
-					.OnClicked(this, &SSyncFilterWindow::OnOkClicked)
+					.Text(LOCTEXT("SaveButtonText", "Save"))
+					.OnClicked(this, &SSyncFilterWindow::OnSaveClicked)
 				]
 				+SUniformGridPanel::Slot(1, 0)
 				[
@@ -155,8 +231,6 @@ void SSyncFilterWindow::Construct(const FArguments& InArgs)
 
 FReply SSyncFilterWindow::OnShowCombinedFilterClicked()
 {
-	fprintf(stderr, "Sync filters:\n%sEnd of sync filters\n", TCHAR_TO_ANSI(*FString::Join(Tab->GetSyncFilters(), TEXT("\n"))));
-
 	TSharedRef<SPopupTextWindow> CombinedFilterWindow = SNew(SPopupTextWindow)
 		.TitleText(LOCTEXT("CombinedSyncFilterWindowTitle", "Combined Sync Filter"))
 		.BodyText(FText::FromString(FString::Join(Tab->GetCombinedSyncFilter(), TEXT("\n"))))
@@ -183,9 +257,35 @@ FReply SSyncFilterWindow::OnCustomViewSyntaxClicked()
 	return FReply::Handled();
 }
 
-FReply SSyncFilterWindow::OnOkClicked()
+FReply SSyncFilterWindow::OnSaveClicked()
 {
-	// Todo: save sync filters
+	TArray<FGuid> SyncExcludedCategoriesCurrent;
+	TArray<bool> EnabledValuesCurrent = SyncFiltersCurrent->GetValues();
+	for (int CategoryIndex = 0; CategoryIndex < WorkspaceCategoriesCurrent.Num(); CategoryIndex++)
+	{
+		if (!EnabledValuesCurrent[CategoryIndex])
+		{
+			SyncExcludedCategoriesCurrent.Add(WorkspaceCategoriesCurrent[CategoryIndex].UniqueId);
+		}
+	}
+
+	TArray<FGuid> SyncExcludedCategoriesAll;
+	TArray<bool> EnabledValuesAll = SyncFiltersAll->GetValues();
+	for (int CategoryIndex = 0; CategoryIndex < WorkspaceCategoriesAll.Num(); CategoryIndex++)
+	{
+		if (!EnabledValuesAll[CategoryIndex])
+		{
+			SyncExcludedCategoriesAll.Add(WorkspaceCategoriesAll[CategoryIndex].UniqueId);
+		}
+	}
+
+	TArray<FString> SyncViewCurrent;
+	CustomSyncViewCurrent->GetPlainText().ToString().ParseIntoArray(SyncViewCurrent, TEXT("\n"));
+
+	TArray<FString> SyncViewAll;
+	CustomSyncViewAll->GetPlainText().ToString().ParseIntoArray(SyncViewAll, TEXT("\n"));
+
+	Tab->OnSyncFilterWindowSaved(SyncViewCurrent, SyncExcludedCategoriesCurrent, SyncViewAll, SyncExcludedCategoriesAll);
 
 	FSlateApplication::Get().FindWidgetWindow(AsShared())->RequestDestroyWindow();
 	return FReply::Handled();
diff --git a/Engine/Source/Programs/SlateUGS/Private/Widgets/SSyncFilterWindow.h b/Engine/Source/Programs/SlateUGS/Private/Widgets/SSyncFilterWindow.h
--- a/Engine/Source/Programs/SlateUGS/Private/Widgets/SSyncFilterWindow.h
+++ b/Engine/Source/Programs/SlateUGS/Private/Widgets/SSyncFilterWindow.h
@@ -3,6 +3,11 @@
 #pragma once
 
 #include "Widgets/SWindow.h"
+#include "Widgets/Input/SMultiLineEditableTextBox.h"
+
+#include "UGSCore/Workspace.h"
+
+#include "SCheckBoxList.h"
 
 class UGSTab;
 
@@ -18,8 +23,19 @@ public:
 private:
 	FReply OnShowCombinedFilterClicked();
 	FReply OnCustomViewSyntaxClicked();
-	FReply OnOkClicked();
+	FReply OnSaveClicked();
 	FReply OnCancelClicked();
 
+	void ConstructSyncFilters();
+	void ConstructCustomSyncViewTextBoxes();
+
 	UGSTab* Tab;
+	TArray<FWorkspaceSyncCategory> WorkspaceCategoriesCurrent;
+	TArray<FWorkspaceSyncCategory> WorkspaceCategoriesAll;
+
+	TSharedPtr<SCheckBoxList> SyncFiltersCurrent;
+	TSharedPtr<SCheckBoxList> SyncFiltersAll;
+	TSharedPtr<SMultiLineEditableTextBox> CustomSyncViewCurrent;
+	TSharedPtr<SMultiLineEditableTextBox> CustomSyncViewAll;
+
 };
