You can now sync to any build by right clicking that build and clicking sync in the dropdown

diff --git a/Engine/Source/Programs/SlateUGS/Private/ChangeInfo.h b/Engine/Source/Programs/SlateUGS/Private/ChangeInfo.h
--- a/Engine/Source/Programs/SlateUGS/Private/ChangeInfo.h
+++ b/Engine/Source/Programs/SlateUGS/Private/ChangeInfo.h
@@ -7,7 +7,7 @@
 struct FChangeInfo
 {
 	EReviewVerdict ReviewStatus;
-	FText Changelist;
+	int Changelist;
 	FText Time;
 	FText Author;
 	FText Description;
diff --git a/Engine/Source/Programs/SlateUGS/Private/UGSTab.cpp b/Engine/Source/Programs/SlateUGS/Private/UGSTab.cpp
--- a/Engine/Source/Programs/SlateUGS/Private/UGSTab.cpp
+++ b/Engine/Source/Programs/SlateUGS/Private/UGSTab.cpp
@@ -246,6 +246,21 @@ bool UGSTab::OnWorkspaceChosen(const FString& Project)
 	return false;
 }
 
+void UGSTab::OnSyncChangelist(int Changelist)
+{
+	TSharedRef<FWorkspaceUpdateContext, ESPMode::ThreadSafe> Context = MakeShared<FWorkspaceUpdateContext, ESPMode::ThreadSafe>(
+		Changelist,
+		Options,
+		CombinedSyncFilter,
+		GetDefaultBuildStepObjects(DetectSettings->NewProjectEditorTarget, UserSettings),
+		ProjectSettings->BuildSteps,
+		TSet<FGuid>(),
+		GetWorkspaceVariables(DetectSettings));
+
+	// Update the workspace with the Context!
+	Workspace->Update(Context);
+}
+
 void UGSTab::OnSyncLatest()
 {
 	int ChangeNumber = -1;
@@ -253,17 +268,7 @@ void UGSTab::OnSyncLatest()
 
 	if (PerforceClient->LatestChangeList(ChangeNumber, AbortEvent, MakeShared<FLogWidgetTextWriter>(GameSyncTabView->GetSyncLog().ToSharedRef()).Get()))
 	{
-		TSharedRef<FWorkspaceUpdateContext, ESPMode::ThreadSafe> Context = MakeShared<FWorkspaceUpdateContext, ESPMode::ThreadSafe>(
-			ChangeNumber,
-			Options,
-			CombinedSyncFilter,
-			GetDefaultBuildStepObjects(DetectSettings->NewProjectEditorTarget, UserSettings),
-			ProjectSettings->BuildSteps,
-			TSet<FGuid>(),
-			GetWorkspaceVariables(DetectSettings));
-
-		// Update the workspace with the Context!
-		Workspace->Update(Context);
+		OnSyncChangelist(ChangeNumber);
 	}
 
 	FPlatformProcess::ReturnSynchEventToPool(AbortEvent);
@@ -411,7 +416,7 @@ void UGSTab::UpdateGameTabBuildList()
 
 				ChangeInfos.Add(MakeShareable(new FChangeInfo{
 					Status,
-					FText::FromString(FString::FromInt(Change.Number)),
+					Change.Number,
 					FText::FromString(DisplayTime.ToString()),
 					FText::FromString(FormatUserName(Change.User)),
 					FText::FromString(DescLines[0])}));
diff --git a/Engine/Source/Programs/SlateUGS/Private/UGSTab.h b/Engine/Source/Programs/SlateUGS/Private/UGSTab.h
--- a/Engine/Source/Programs/SlateUGS/Private/UGSTab.h
+++ b/Engine/Source/Programs/SlateUGS/Private/UGSTab.h
@@ -35,6 +35,7 @@ public:
 
 	// Slate callbacks
 	bool OnWorkspaceChosen(const FString& Project);
+	void OnSyncChangelist(int Changelist);
 	void OnSyncLatest();
 	void OnSyncFilterWindowSaved(
 		const TArray<FString>& SyncViewCurrent,
diff --git a/Engine/Source/Programs/SlateUGS/Private/Widgets/SGameSyncTab.cpp b/Engine/Source/Programs/SlateUGS/Private/Widgets/SGameSyncTab.cpp
--- a/Engine/Source/Programs/SlateUGS/Private/Widgets/SGameSyncTab.cpp
+++ b/Engine/Source/Programs/SlateUGS/Private/Widgets/SGameSyncTab.cpp
@@ -3,6 +3,7 @@
 #include "SGameSyncTab.h"
 
 #include "Framework/Application/SlateApplication.h"
+#include "Framework/MultiBox/MultiBoxBuilder.h"
 
 #include "SPositiveActionButton.h"
 #include "SSimpleComboButton.h"
@@ -38,7 +39,7 @@ void SBuildDataRow::Construct(const FArguments& InArgs, const TSharedRef<STableV
 	SMultiColumnTableRow<TSharedPtr<FChangeInfo>>::Construct(SMultiColumnTableRow::FArguments(), InOwnerTableView);
 }
 
-TSharedRef<SWidget> SBuildDataRow::GenerateWidgetForColumn(const FName& ColumnId)
+TSharedRef<SWidget> SBuildDataRow::GenerateWidgetForColumn(const FName& ColumnId) // Todo: maybe can refactor some of this code so there's less duplication by using the root SWidget class on different types
 {
 	if (ColumnId == HordeTableColumnStatus)
 	{
@@ -73,7 +74,7 @@ TSharedRef<SWidget> SBuildDataRow::GenerateWidgetForColumn(const FName& ColumnId
 	TSharedRef<STextBlock> TextItem = SNew(STextBlock);
 	if (ColumnId == HordeTableColumnChange)
 	{
-		TextItem->SetText(CurrentItem->Changelist);
+		TextItem->SetText(FText::FromString(FString::FromInt(CurrentItem->Changelist)));
 		TextItem->SetJustification(ETextJustify::Center);
 	}
 	if (ColumnId == HordeTableColumnTime)
@@ -91,7 +92,9 @@ TSharedRef<SWidget> SBuildDataRow::GenerateWidgetForColumn(const FName& ColumnId
 	}
 
 	return SNew(SBox)
-		.Padding(10.0f, 0.0f)
+		.HAlign(HAlign_Center)
+		.VAlign(VAlign_Center)
+		.Padding(10.0f, 2.5f)
 		[
 			TextItem
 		];
@@ -110,12 +113,7 @@ TSharedRef<SWidget> SGameSyncTab::MakeSyncButtonDropdown()
 		FText::FromString(TEXT("Sync Latest")),
 		FText::FromString(TEXT("Sync to the latest submitted changelist")),
 		FSlateIcon(),
-		FUIAction(
-			FExecuteAction::CreateLambda([this]()
-			{
-				Tab->OnSyncLatest();
-			}),
-			FCanExecuteAction::CreateLambda([] { return true; })),
+		FUIAction(FExecuteAction::CreateRaw(Tab, &UGSTab::OnSyncLatest)),
 		NAME_None,
 		EUserInterfaceActionType::Button
 		);
@@ -346,9 +344,11 @@ void SGameSyncTab::Construct(const FArguments& InArgs)
 		.FillHeight(0.45f)
 		[
 			SAssignNew(HordeBuildsView, SListView<TSharedPtr<FChangeInfo>>)
+			.ListItemsSource(&HordeBuilds)
+			.SelectionMode(ESelectionMode::Single)
 			.IsEnabled_Lambda([this] { return !Tab->IsSyncing(); })
 			.OnGenerateRow(this, &SGameSyncTab::GenerateHordeBuildTableRow)
-			.ListItemsSource(&HordeBuilds)
+			.OnContextMenuOpening(this, &SGameSyncTab::OnRightClickedBuild)
 			.HeaderRow(
 				SNew(SHeaderRow)
 				+SHeaderRow::Column(HordeTableColumnStatus)
@@ -418,8 +418,28 @@ void SGameSyncTab::SetProjectPathText(FText ProjectPath)
 
 void SGameSyncTab::AddHordeBuilds(const TArray<TSharedPtr<FChangeInfo>>& Builds)
 {
-	HordeBuilds += Builds;
+	HordeBuilds = Builds;
 	HordeBuildsView->RebuildList();
 }
 
+TSharedPtr<SWidget> SGameSyncTab::OnRightClickedBuild()
+{
+	FMenuBuilder MenuBuilder(true, nullptr);
+
+	TArray<TSharedPtr<FChangeInfo>> SelectedItems = HordeBuildsView->GetSelectedItems(); // Todo: since I disabled multi select, I might be able to assume the array size is always 1?
+	if (SelectedItems.IsValidIndex(0))
+	{
+		MenuBuilder.AddMenuEntry(
+			FText::FromString(TEXT("Sync")),
+			FText::FromString("Sync to CL " + FString::FromInt(SelectedItems[0]->Changelist)),
+			FSlateIcon(),
+			FUIAction(FExecuteAction::CreateRaw(Tab, &UGSTab::OnSyncChangelist, SelectedItems[0]->Changelist)),
+			NAME_None,
+			EUserInterfaceActionType::Button
+		);
+	}
+
+	return MenuBuilder.MakeWidget();
+}
+
 #undef LOCTEXT_NAMESPACE
diff --git a/Engine/Source/Programs/SlateUGS/Private/Widgets/SGameSyncTab.h b/Engine/Source/Programs/SlateUGS/Private/Widgets/SGameSyncTab.h
--- a/Engine/Source/Programs/SlateUGS/Private/Widgets/SGameSyncTab.h
+++ b/Engine/Source/Programs/SlateUGS/Private/Widgets/SGameSyncTab.h
@@ -44,6 +44,9 @@ private:
 	TSharedRef<ITableRow> GenerateHordeBuildTableRow(TSharedPtr<FChangeInfo> InItem, const TSharedRef<STableViewBase>& InOwnerTable);
 	TSharedRef<SWidget> MakeSyncButtonDropdown();
 
+	// Widget callbacks
+	TSharedPtr<SWidget> OnRightClickedBuild();
+
 	TSharedPtr<SListView<TSharedPtr<FChangeInfo>>> HordeBuildsView;
 	TArray<TSharedPtr<FChangeInfo>> HordeBuilds;
 
