Disabled currently-unused buttons

diff --git a/Engine/Source/Programs/SlateUGS/Private/Widgets/SGameSyncTab.cpp b/Engine/Source/Programs/SlateUGS/Private/Widgets/SGameSyncTab.cpp
--- a/Engine/Source/Programs/SlateUGS/Private/Widgets/SGameSyncTab.cpp
+++ b/Engine/Source/Programs/SlateUGS/Private/Widgets/SGameSyncTab.cpp
@@ -92,7 +92,6 @@ TSharedRef<SWidget> SBuildDataRow::GenerateWidgetForColumn(const FName& ColumnId
 	}
 
 	return SNew(SBox)
-		.HAlign(HAlign_Center)
 		.VAlign(VAlign_Center)
 		.Padding(10.0f, 2.5f)
 		[
@@ -160,6 +159,7 @@ void SGameSyncTab::Construct(const FArguments& InArgs)
 					[
 						SNew(SSimpleButton)
 						.Text(LOCTEXT("Build", "Build"))
+						.IsEnabled(false) // Todo: enable after adding this functionality
 					]
 					+SHorizontalBox::Slot()
 					.AutoWidth()
@@ -168,6 +168,7 @@ void SGameSyncTab::Construct(const FArguments& InArgs)
 						.Text(LOCTEXT("RunUnrealEditor", "Run Unreal Editor"))
 						.Icon(FAppStyle::Get().GetBrush("Icons.Launch"))
 						.HasDownArrow(true)
+						.IsEnabled(false) // Todo: enable after adding this functionality
 					]
 					+SHorizontalBox::Slot()
 					.AutoWidth()
@@ -175,6 +176,7 @@ void SGameSyncTab::Construct(const FArguments& InArgs)
 						SNew(SSimpleButton)
 						.Text(LOCTEXT("OpenSolution", "OpenSolution"))
 						.Icon(FAppStyle::Get().GetBrush("MainFrame.OpenVisualStudio")) // Todo: shouldn't use this icon (repurposing, also could use other IDEs)
+						.IsEnabled(false) // Todo: enable after adding this functionality
 					]
 				]
 				+SHorizontalBox::Slot()
@@ -186,6 +188,7 @@ void SGameSyncTab::Construct(const FArguments& InArgs)
 					[
 						SNew(SSimpleButton)
 						.Text(LOCTEXT("BuildHealth", "Build Health")) // Todo: What icon?
+						.IsEnabled(false) // Todo: enable after adding this functionality
 					]
 					+SHorizontalBox::Slot()
 					.AutoWidth()
@@ -193,6 +196,7 @@ void SGameSyncTab::Construct(const FArguments& InArgs)
 						SNew(SSimpleButton)
 						.Text(LOCTEXT("SDKInfo", "SDK Info"))
 						.Icon(FAppStyle::Get().GetBrush("Icons.Settings")) // Todo: What icon?
+						.IsEnabled(false) // Todo: enable after adding this functionality
 					]
 					+SHorizontalBox::Slot()
 					.AutoWidth()
@@ -200,6 +204,7 @@ void SGameSyncTab::Construct(const FArguments& InArgs)
 						SNew(SSimpleButton)
 						.Text(LOCTEXT("OpenPerforce", "Open Perforce"))
 						.Icon(FAppStyle::Get().GetBrush("Icons.Blueprints")) // Todo: shouldn't use this icon (repurposing)
+						.IsEnabled(false) // Todo: enable after adding this functionality
 					]
 					+SHorizontalBox::Slot()
 					.AutoWidth()
@@ -207,6 +212,7 @@ void SGameSyncTab::Construct(const FArguments& InArgs)
 						SNew(SSimpleButton)
 						.Text(LOCTEXT("CleanSolution", "Clean Solution"))
 						.Icon(FAppStyle::Get().GetBrush("GraphEditor.Clean")) // Todo: shouldn't use this icon (repurposing)
+						.IsEnabled(false) // Todo: enable after adding this functionality
 					]
 					+SHorizontalBox::Slot()
 					.AutoWidth()
@@ -214,6 +220,7 @@ void SGameSyncTab::Construct(const FArguments& InArgs)
 						SNew(SSimpleButton)
 						.Text(LOCTEXT("Filter", "Filter"))
 						.Icon(FAppStyle::Get().GetBrush("Icons.Filter"))
+						// Todo: this is probably the wrong "Filter" button. The functionality below should probably be in the settings dropdown
 						.OnClicked_Lambda([this] { FSlateApplication::Get().AddModalWindow(SNew(SSyncFilterWindow).Tab(Tab), Tab->GetTabArgs().GetOwnerWindow(), false); return FReply::Handled(); })
 					]
 					+SHorizontalBox::Slot()
@@ -223,6 +230,7 @@ void SGameSyncTab::Construct(const FArguments& InArgs)
 						SNew(SSimpleButton)
 						.Text(LOCTEXT("Settings", "Settings"))
 						.Icon(FAppStyle::Get().GetBrush("Icons.Settings"))
+						.IsEnabled(false) // Todo: enable after adding this functionality
 					]
 				]
 			]
diff --git a/Engine/Source/Programs/SlateUGS/Private/Widgets/SWorkspaceWindow.cpp b/Engine/Source/Programs/SlateUGS/Private/Widgets/SWorkspaceWindow.cpp
--- a/Engine/Source/Programs/SlateUGS/Private/Widgets/SWorkspaceWindow.cpp
+++ b/Engine/Source/Programs/SlateUGS/Private/Widgets/SWorkspaceWindow.cpp
@@ -93,6 +93,7 @@ void SWorkspaceWindow::Construct(const FArguments& InArgs)
 					.AutoWidth()
 					[
 						SNew(SCheckBox).Style(FAppStyle::Get(), "RadioButton")
+						.IsEnabled(false) // Todo: enable after adding this functionality
 						.IsChecked_Lambda([this] () { return bIsLocalFileSelected ? ECheckBoxState::Unchecked : ECheckBoxState::Checked; })
 						.OnCheckStateChanged_Lambda( [this] (ECheckBoxState InState) { bIsLocalFileSelected = (!bIsLocalFileSelected && InState == ECheckBoxState::Checked); } )
 					]
