Fixed problem with not appending previous log.
Fixed problem with some build list entries taking up 2 lines instead of 1

diff --git a/Engine/Source/Programs/SlateUGS/Private/UGSTab.cpp b/Engine/Source/Programs/SlateUGS/Private/UGSTab.cpp
--- a/Engine/Source/Programs/SlateUGS/Private/UGSTab.cpp
+++ b/Engine/Source/Programs/SlateUGS/Private/UGSTab.cpp
@@ -419,7 +419,7 @@ void UGSTab::UpdateGameTabBuildList()
 					Change.Number,
 					FText::FromString(DisplayTime.ToString()),
 					FText::FromString(FormatUserName(Change.User)),
-					FText::FromString(DescLines[0])}));
+					FText::FromString(DescLines[0].TrimStartAndEnd())}));
 			}
 		}
 	}
@@ -539,7 +539,8 @@ void UGSTab::SetupWorkspace()
 	FString TelemetryProjectIdentifier = FPerforceUtils::GetClientOrDepotDirectoryName(*DetectSettings->NewSelectedProjectIdentifier);
 
 	FString LogFileName = DataFolder / FPaths::GetPath(ProjectFileName) + TEXT(".sync.log");
-	GameSyncTabView->SetSyncLogLocation(LogFileName);
+	GameSyncTabView->SetSyncLogLocation(LogFileName); // Todo: if SetSyncLogLocation fails, then it failed to create a log file and may need to handle that
+	GameSyncTabView->GetSyncLog()->AppendLine(TEXT("Creating log at: ") + LogFileName);
 
 	Workspace = MakeShared<FWorkspace>(
 		PerforceClient.ToSharedRef(),
diff --git a/Engine/Source/Programs/SlateUGS/Private/Widgets/SGameSyncTab.cpp b/Engine/Source/Programs/SlateUGS/Private/Widgets/SGameSyncTab.cpp
--- a/Engine/Source/Programs/SlateUGS/Private/Widgets/SGameSyncTab.cpp
+++ b/Engine/Source/Programs/SlateUGS/Private/Widgets/SGameSyncTab.cpp
@@ -404,9 +404,9 @@ TSharedPtr<SLogWidget> SGameSyncTab::GetSyncLog() const
 	return SyncLog;
 }
 
-void SGameSyncTab::SetSyncLogLocation(const FString& LogFileName)
+bool SGameSyncTab::SetSyncLogLocation(const FString& LogFileName)
 {
-    SyncLog->OpenFile(*LogFileName);
+    return SyncLog->OpenFile(*LogFileName);
 }
 
 void SGameSyncTab::SetStreamPathText(FText StreamPath)
diff --git a/Engine/Source/Programs/SlateUGS/Private/Widgets/SGameSyncTab.h b/Engine/Source/Programs/SlateUGS/Private/Widgets/SGameSyncTab.h
--- a/Engine/Source/Programs/SlateUGS/Private/Widgets/SGameSyncTab.h
+++ b/Engine/Source/Programs/SlateUGS/Private/Widgets/SGameSyncTab.h
@@ -32,7 +32,7 @@ public:
 	// We need access to the SyncLog when creating the Workspace
 	// TODO: think of a better way to do this
 	TSharedPtr<SLogWidget> GetSyncLog() const;
-	void SetSyncLogLocation(const FString& LogFileName);
+	bool SetSyncLogLocation(const FString& LogFileName);
 
 	void SetStreamPathText(FText StreamPath);
 	void SetChangelistText(FText Changelist);
diff --git a/Engine/Source/Programs/SlateUGS/Private/Widgets/SLogWidget.cpp b/Engine/Source/Programs/SlateUGS/Private/Widgets/SLogWidget.cpp
--- a/Engine/Source/Programs/SlateUGS/Private/Widgets/SLogWidget.cpp
+++ b/Engine/Source/Programs/SlateUGS/Private/Widgets/SLogWidget.cpp
@@ -86,14 +86,16 @@ void SLogWidget::Construct(const FArguments& InArgs)
 	[
 		SNew(SBorder)
 		[
-			 SAssignNew(MessagesTextBox, SMultiLineEditableTextBox)
+			SAssignNew(MessagesTextBox, SMultiLineEditableTextBox)
 				.Style(FAppStyle::Get(), "Log.TextBox")
 				.TextStyle(FAppStyle::Get(), "Log.Normal")
+				.Padding(10.0f)
 				.ForegroundColor(FLinearColor::Gray)
 				.Marshaller(MessagesTextMarshaller)
 				.IsReadOnly(true)
 				.AlwaysShowScrollbars(true)
 				.OnVScrollBarUserScrolled(this, &SLogWidget::OnScroll)
+
 		]
 	];
 
@@ -103,13 +105,14 @@ void SLogWidget::Construct(const FArguments& InArgs)
 bool SLogWidget::OpenFile(const TCHAR* NewLogFileName)
 {
 	CloseFile();
-	Clear();
+	// Clear(); // Todo: calling Clear() breaks some things (like causing the logs below to not be appended), but maybe we're doing something wrong?
 
 	TArray<FString> InitialLines;
 	if (FFileHelper::LoadFileToStringArray(InitialLines, NewLogFileName))
 	{
 		for (const FString& InitialLine : InitialLines)
 		{
+			fprintf(stderr, "\nLine appended: %s\n", TCHAR_TO_ANSI(*InitialLine));
 			AppendLine(InitialLine);
 		}
 	}
