Remove debugging print

diff --git a/Engine/Source/Programs/SlateUGS/Private/Widgets/SLogWidget.cpp b/Engine/Source/Programs/SlateUGS/Private/Widgets/SLogWidget.cpp
--- a/Engine/Source/Programs/SlateUGS/Private/Widgets/SLogWidget.cpp
+++ b/Engine/Source/Programs/SlateUGS/Private/Widgets/SLogWidget.cpp
@@ -112,7 +112,6 @@ bool SLogWidget::OpenFile(const TCHAR* NewLogFileName)
 	{
 		for (const FString& InitialLine : InitialLines)
 		{
-			fprintf(stderr, "\nLine appended: %s\n", TCHAR_TO_ANSI(*InitialLine));
 			AppendLine(InitialLine);
 		}
 	}
