Read in PanelColor/AlertMessage from the settings

diff --git a/Engine/Source/Programs/SlateUGS/Private/UGSCore/DetectProjectSettingsTask.cpp b/Engine/Source/Programs/SlateUGS/Private/UGSCore/DetectProjectSettingsTask.cpp
--- a/Engine/Source/Programs/SlateUGS/Private/UGSCore/DetectProjectSettingsTask.cpp
+++ b/Engine/Source/Programs/SlateUGS/Private/UGSCore/DetectProjectSettingsTask.cpp
@@ -216,7 +216,7 @@ TSharedRef<FModalTaskResult> FDetectProjectSettingsTask::RunInternal(FEvent* Abo
 		{
 			return FModalTaskResult::Failure(FText::Format(LOCTEXT("FailedToGetStreamInfo", "Failed to get stream info for {0}"), FText::FromString(StreamName)));
 		}
-		NewSelectedProjectIdentifier = FString::Printf(TEXT("{0}/{1}"), *StreamPrefix, *NewSelectedClientFileName.Mid(ExpectedPrefix.Len()));
+		NewSelectedProjectIdentifier = FString::Printf(TEXT("%s/%s"), *StreamPrefix, *NewSelectedClientFileName.Mid(ExpectedPrefix.Len()));
 	}
 	else
 	{
diff --git a/Engine/Source/Programs/SlateUGS/Private/UGSCore/Workspace.cpp b/Engine/Source/Programs/SlateUGS/Private/UGSCore/Workspace.cpp
--- a/Engine/Source/Programs/SlateUGS/Private/UGSCore/Workspace.cpp
+++ b/Engine/Source/Programs/SlateUGS/Private/UGSCore/Workspace.cpp
@@ -145,13 +145,14 @@ FWorkspaceSyncCategory::FWorkspaceSyncCategory(const FGuid& InUniqueId, const TC
 
 //// FWorkspace ////
 
-FWorkspace::FWorkspace(TSharedRef<FPerforceConnection> InPerforce, const FString& InLocalRootPath, const FString& InSelectedLocalFileName, const FString& InClientRootPath, const FString& InSelectedClientFileName, int InInitialChangeNumber, int InLastBuiltChangeNumber, const FString& InTelemetryProjectPath, TSharedRef<FLineBasedTextWriter> InLog)
+FWorkspace::FWorkspace(TSharedRef<FPerforceConnection> InPerforce, const FString& InLocalRootPath, const FString& InSelectedLocalFileName, const FString& InClientRootPath, const FString& InSelectedClientFileName, const FString& InSelectedProjectIdentifier, int InInitialChangeNumber, int InLastBuiltChangeNumber, const FString& InTelemetryProjectPath, TSharedRef<FLineBasedTextWriter> InLog)
 	: Perforce(InPerforce)
 	, SyncPaths(GetSyncPaths(InClientRootPath, InSelectedClientFileName))
 	, LocalRootPath(InLocalRootPath)
 	, SelectedLocalFileName(InSelectedLocalFileName)
 	, ClientRootPath(InClientRootPath)
 	, SelectedClientFileName(InSelectedClientFileName)
+	, SelectedProjectIdentifier(InSelectedProjectIdentifier)
 	, TelemetryProjectPath(InTelemetryProjectPath)
 	, CurrentChangeNumber(InInitialChangeNumber)
 	, PendingChangeNumber(InInitialChangeNumber)
@@ -163,6 +164,7 @@ FWorkspace::FWorkspace(TSharedRef<FPerforceConnection> InPerforce, const FString
 	, WorkerThread(nullptr)
 {
 	ProjectStreamFilter = ReadProjectStreamFilter(InPerforce.Get(), ProjectConfigFile.Get(), AbortEvent, InLog.Get());
+	UpdateStatusPanel();
 }
 
 FWorkspace::~FWorkspace()
@@ -170,6 +172,27 @@ FWorkspace::~FWorkspace()
 	CancelUpdate();
 }
 
+void FWorkspace::UpdateStatusPanel()
+{
+	TSharedPtr<const FCustomConfigSection> ProjectSection = ProjectConfigFile->FindSection(*SelectedProjectIdentifier);
+
+	if (ProjectSection.IsValid())
+	{
+		PanelColor   = ProjectSection->GetValue(TEXT("StatusPanelColor"), TEXT(""));
+		AlertMessage = ProjectSection->GetValue(TEXT("Message"), TEXT(""));
+	}
+}
+
+FString FWorkspace::GetPanelColor() const
+{
+	return PanelColor;
+}
+
+FString FWorkspace::GetAlertMessage() const
+{
+	return AlertMessage;
+}
+
 TMap<FGuid, FWorkspaceSyncCategory> FWorkspace::GetSyncCategories() const
 {
 	TMap<FGuid, FWorkspaceSyncCategory> UniqueIdToCategory;
diff --git a/Engine/Source/Programs/SlateUGS/Private/UGSCore/Workspace.h b/Engine/Source/Programs/SlateUGS/Private/UGSCore/Workspace.h
--- a/Engine/Source/Programs/SlateUGS/Private/UGSCore/Workspace.h
+++ b/Engine/Source/Programs/SlateUGS/Private/UGSCore/Workspace.h
@@ -78,11 +78,12 @@ public:
 	const FString SelectedLocalFileName;
 	const FString ClientRootPath;
 	const FString SelectedClientFileName;
+	const FString SelectedProjectIdentifier;
 	const FString TelemetryProjectPath;
 
 	TFunction<void(TSharedRef<FWorkspaceUpdateContext, ESPMode::ThreadSafe>, EWorkspaceUpdateResult, const FString&)> OnUpdateComplete;
 
-	FWorkspace(TSharedRef<FPerforceConnection> InPerforce, const FString& InLocalRootPath, const FString& InSelectedLocalFileName, const FString& InClientRootPath, const FString& InSelectedClientFileName, int InInitialChangeNumber, int InLastBuiltChangeNumber, const FString& InTelemetryProjectPath, TSharedRef<FLineBasedTextWriter> InLog);
+	FWorkspace(TSharedRef<FPerforceConnection> InPerforce, const FString& InLocalRootPath, const FString& InSelectedLocalFileName, const FString& InClientRootPath, const FString& InSelectedClientFileName, const FString& InSelectedProjectIdentifier, int InInitialChangeNumber, int InLastBuiltChangeNumber, const FString& InTelemetryProjectPath, TSharedRef<FLineBasedTextWriter> InLog);
 	~FWorkspace();
 
 	TMap<FGuid, FWorkspaceSyncCategory> GetSyncCategories() const;
@@ -98,6 +99,9 @@ public:
 	void Update(const TSharedRef<FWorkspaceUpdateContext, ESPMode::ThreadSafe>& Context);
 	void CancelUpdate();
 
+	FString GetPanelColor() const;
+	FString GetAlertMessage() const;
+
 private:
 	static const TCHAR* DefaultBuildTargets[];
 	static const FWorkspaceSyncCategory DefaultSyncCategories[];
@@ -118,6 +122,9 @@ private:
 	FRunnableThread* WorkerThread;
 	FProgressValue Progress;
 
+	FString PanelColor;
+	FString AlertMessage;
+	
 	static FWorkspace* ActiveWorkspace;
 
 	virtual uint32 Run() override;
@@ -128,6 +135,8 @@ private:
 	static TSharedRef<FCustomConfigFile, ESPMode::ThreadSafe> ReadProjectConfigFile(const FString& LocalRootPath, const FString& SelectedLocalFileName, FLineBasedTextWriter& Log);
 	static TArray<FString> ReadProjectStreamFilter(FPerforceConnection& Perforce, const FCustomConfigFile& ProjectConfigFile, FEvent* AbortEvent, FLineBasedTextWriter& Log);
 	static FString FormatTime(long Seconds);
+
+	void UpdateStatusPanel();
 	
 	bool HasModifiedSourceFiles() const;
 	bool FindUnresolvedFiles(TArray<FPerforceFileRecord>& OutUnresolvedFiles) const;
diff --git a/Engine/Source/Programs/SlateUGS/Private/UGSTab.cpp b/Engine/Source/Programs/SlateUGS/Private/UGSTab.cpp
--- a/Engine/Source/Programs/SlateUGS/Private/UGSTab.cpp
+++ b/Engine/Source/Programs/SlateUGS/Private/UGSTab.cpp
@@ -537,6 +537,7 @@ void UGSTab::SetupWorkspace()
 
 	FString ProjectLogBaseName = DataFolder / FString::Printf(TEXT("%s@%s"), *PerforceClient->ClientName, *ClientKey.Replace(TEXT("/"), TEXT("$")));
 	FString TelemetryProjectIdentifier = FPerforceUtils::GetClientOrDepotDirectoryName(*DetectSettings->NewSelectedProjectIdentifier);
+	FString SelectedProjectIdentifier = DetectSettings->NewSelectedProjectIdentifier;
 
 	FString LogFileName = DataFolder / FPaths::GetPath(ProjectFileName) + TEXT(".sync.log");
 	GameSyncTabView->SetSyncLogLocation(LogFileName); // Todo: if SetSyncLogLocation fails, then it failed to create a log file and may need to handle that
@@ -548,6 +549,7 @@ void UGSTab::SetupWorkspace()
 		ProjectFileName,
 		DetectSettings->BranchClientPath,
 		DetectSettings->NewSelectedClientFileName,
+		SelectedProjectIdentifier,
 		CurrentChangeNumber,
 		WorkspaceSettings->LastBuiltChangeNumber,
 		TelemetryProjectIdentifier,
@@ -590,7 +592,6 @@ void UGSTab::SetupWorkspace()
 	// Setup our Perforce and Event monitoring threads
 	FString BranchClientPath = DetectSettings->BranchDirectoryName;
 	FString SelectedClientFileName = DetectSettings->NewSelectedClientFileName;
-	FString SelectedProjectIdentifier = DetectSettings->NewSelectedProjectIdentifier;
 
 	// TODO create callback functions that will be queued for the main thread to generate and update the main table view
 	PerforceMonitor = MakeShared<FPerforceMonitor>(PerforceClient.ToSharedRef(), BranchClientPath, SelectedClientFileName, SelectedProjectIdentifier, ProjectLogBaseName + ".p4.log");
