Implement missing FDatetTime parsing and remove ensure to fix that

diff --git a/Engine/Source/Programs/SlateUGS/Private/UGSCore/Perforce.cpp b/Engine/Source/Programs/SlateUGS/Private/UGSCore/Perforce.cpp
--- a/Engine/Source/Programs/SlateUGS/Private/UGSCore/Perforce.cpp
+++ b/Engine/Source/Programs/SlateUGS/Private/UGSCore/Perforce.cpp
@@ -683,7 +683,10 @@ bool FPerforceConnection::TryParseFileChangeSummary(const TArray<FString>& Lines
 		return false;
 	}
 
-	if(!FUtility::TryParse(*Tokens[0].Mid(1), OutChange.Revision) || !FUtility::TryParse(*Tokens[2], OutChange.ChangeNumber) || ensure(false))//!DateTime.TryParse(Tokens[5] + " " + Tokens[6], out Change.Date))
+	// Replace [5] date and [6] time with . as FDateTime expects format to be (yyyy.mm.dd-hh.mm.ss)
+	Tokens[5].ReplaceCharInline('/', '.');
+	Tokens[6].ReplaceCharInline(':', '.');
+	if(!FUtility::TryParse(*Tokens[0].Mid(1), OutChange.Revision) || !FUtility::TryParse(*Tokens[2], OutChange.ChangeNumber) || !FDateTime::Parse(Tokens[5] + TEXT("-") + Tokens[6], OutChange.Date))
 	{
 		return false;
 	}
