Move UGSCore into a namespace

diff --git a/Engine/Source/Programs/SlateUGS/Private/ChangeInfo.h b/Engine/Source/Programs/SlateUGS/Private/ChangeInfo.h
--- a/Engine/Source/Programs/SlateUGS/Private/ChangeInfo.h
+++ b/Engine/Source/Programs/SlateUGS/Private/ChangeInfo.h
@@ -9,7 +9,7 @@ struct FChangeInfo
 	FDateTime Time;
 	bool bHeaderRow = false;
 
-	EReviewVerdict ReviewStatus = EReviewVerdict::Unknown;
+	UGSCore::EReviewVerdict ReviewStatus = UGSCore::EReviewVerdict::Unknown;
 	int Changelist = 0;
 	FText Author;
 	FText Description;
diff --git a/Engine/Source/Programs/SlateUGS/Private/UGSCore/BuildStep.cpp b/Engine/Source/Programs/SlateUGS/Private/UGSCore/BuildStep.cpp
--- a/Engine/Source/Programs/SlateUGS/Private/UGSCore/BuildStep.cpp
+++ b/Engine/Source/Programs/SlateUGS/Private/UGSCore/BuildStep.cpp
@@ -2,6 +2,9 @@
 
 #include "BuildStep.h"
 
+namespace UGSCore
+{
+
 const TCHAR* FBuildStep::UniqueIdKey = TEXT("UniqueId");
 
 bool TryParse(const TCHAR* Text, EBuildStepType& OutType)
@@ -152,3 +155,5 @@ FCustomConfigObject FBuildStep::ToConfigObject(const FCustomConfigObject* Defaul
 	Result.AddOverrides(ToConfigObject(), DefaultObject);
 	return Result;
 }
+
+} // namespace UGSCore
diff --git a/Engine/Source/Programs/SlateUGS/Private/UGSCore/BuildStep.h b/Engine/Source/Programs/SlateUGS/Private/UGSCore/BuildStep.h
--- a/Engine/Source/Programs/SlateUGS/Private/UGSCore/BuildStep.h
+++ b/Engine/Source/Programs/SlateUGS/Private/UGSCore/BuildStep.h
@@ -6,6 +6,9 @@
 #include "Misc/Guid.h"
 #include "CustomConfigFile.h"
 
+namespace UGSCore
+{
+
 enum class EBuildStepType
 {
 	Compile,
@@ -44,3 +47,5 @@ struct FBuildStep
 	FCustomConfigObject ToConfigObject() const;
 	FCustomConfigObject ToConfigObject(const FCustomConfigObject* DefaultObject) const;
 };
+
+} // namespace UGSCore
diff --git a/Engine/Source/Programs/SlateUGS/Private/UGSCore/CustomConfigFile.cpp b/Engine/Source/Programs/SlateUGS/Private/UGSCore/CustomConfigFile.cpp
--- a/Engine/Source/Programs/SlateUGS/Private/UGSCore/CustomConfigFile.cpp
+++ b/Engine/Source/Programs/SlateUGS/Private/UGSCore/CustomConfigFile.cpp
@@ -4,6 +4,9 @@
 #include "Misc/FileHelper.h"
 #include "Utility.h"
 
+namespace UGSCore
+{
+
 static const TCHAR* ConfigSeparatorCharacters = TEXT("(),= \t\"");
 
 //// FCustomConfigObject ////
@@ -734,3 +737,5 @@ bool FCustomConfigFile::TryGetValues(const TCHAR* Key, TArray<FGuid> &Values) co
 	TSharedPtr<const FCustomConfigSection> Section = FindSection(*FString(Dot - Key, Key));
 	return Section.IsValid()? Section->TryGetValues(Dot + 1, Values) : false;
 }
+
+} // namespace UGSCore
diff --git a/Engine/Source/Programs/SlateUGS/Private/UGSCore/CustomConfigFile.h b/Engine/Source/Programs/SlateUGS/Private/UGSCore/CustomConfigFile.h
--- a/Engine/Source/Programs/SlateUGS/Private/UGSCore/CustomConfigFile.h
+++ b/Engine/Source/Programs/SlateUGS/Private/UGSCore/CustomConfigFile.h
@@ -6,6 +6,9 @@
 #include "Misc/Guid.h"
 #include "Templates/Tuple.h"
 
+namespace UGSCore
+{
+
 struct FCustomConfigObject
 {
 	TArray<TTuple<FString, FString>> Pairs;
@@ -96,3 +99,5 @@ struct FCustomConfigFile
 	bool TryGetValues(const TCHAR* Key, TArray<FString>& Values) const;
 	bool TryGetValues(const TCHAR* Key, TArray<FGuid>& Value) const;
 };
+
+} // namespace UGSCore
diff --git a/Engine/Source/Programs/SlateUGS/Private/UGSCore/DetectProjectSettingsTask.cpp b/Engine/Source/Programs/SlateUGS/Private/UGSCore/DetectProjectSettingsTask.cpp
--- a/Engine/Source/Programs/SlateUGS/Private/UGSCore/DetectProjectSettingsTask.cpp
+++ b/Engine/Source/Programs/SlateUGS/Private/UGSCore/DetectProjectSettingsTask.cpp
@@ -8,6 +8,9 @@
 
 #define LOCTEXT_NAMESPACE "UnrealGameSync"
 
+namespace UGSCore
+{
+
 FDetectProjectSettingsTask::FDetectProjectSettingsTask(TSharedRef<FPerforceConnection> InPerforce, const FString& InNewSelectedFileName, TSharedRef<FLineBasedTextWriter> InLog)
 	: NewSelectedFileName(InNewSelectedFileName)
 	, Perforce(MoveTemp(InPerforce))
@@ -272,3 +275,5 @@ bool FDetectProjectSettingsTask::TryGetStreamPrefix(TSharedRef<FPerforceConnecti
 }
 
 #undef LOCTEXT_NAMESPACE
+
+} // namespace UGSCore
diff --git a/Engine/Source/Programs/SlateUGS/Private/UGSCore/DetectProjectSettingsTask.h b/Engine/Source/Programs/SlateUGS/Private/UGSCore/DetectProjectSettingsTask.h
--- a/Engine/Source/Programs/SlateUGS/Private/UGSCore/DetectProjectSettingsTask.h
+++ b/Engine/Source/Programs/SlateUGS/Private/UGSCore/DetectProjectSettingsTask.h
@@ -7,6 +7,9 @@
 #include "Perforce.h"
 #include "OutputAdapters.h"
 
+namespace UGSCore
+{
+
 class FDetectProjectSettingsTask : public IModalTask
 {
 public:
@@ -34,3 +37,5 @@ private:
 	TSharedRef<FModalTaskResult> RunInternal(FEvent* AbortEvent);
 	static bool TryGetStreamPrefix(TSharedRef<FPerforceConnection> Perforce, const FString& StreamName, FEvent* AbortEvent, FLineBasedTextWriter& Log, FString& OutStreamPrefix);
 };
+
+} // namespace UGSCore
diff --git a/Engine/Source/Programs/SlateUGS/Private/UGSCore/EventMonitor.cpp b/Engine/Source/Programs/SlateUGS/Private/UGSCore/EventMonitor.cpp
--- a/Engine/Source/Programs/SlateUGS/Private/UGSCore/EventMonitor.cpp
+++ b/Engine/Source/Programs/SlateUGS/Private/UGSCore/EventMonitor.cpp
@@ -5,6 +5,9 @@
 #include "HAL/PlatformProcess.h"
 #include "Misc/ScopeLock.h"
 
+namespace UGSCore
+{
+
 //// FEventData ////
 
 FEventData::FEventData(int32 InChangeNumber, const FString& InUserName, EEventType InType, const FString& InProject)
@@ -824,3 +827,5 @@ void FEventMonitor::UpdateActiveInvestigations()
 		}
 	}
 }
+
+} // namespace UGSCore
diff --git a/Engine/Source/Programs/SlateUGS/Private/UGSCore/EventMonitor.h b/Engine/Source/Programs/SlateUGS/Private/UGSCore/EventMonitor.h
--- a/Engine/Source/Programs/SlateUGS/Private/UGSCore/EventMonitor.h
+++ b/Engine/Source/Programs/SlateUGS/Private/UGSCore/EventMonitor.h
@@ -7,6 +7,9 @@
 #include "HAL/Runnable.h"
 #include "HAL/RunnableThread.h"
 
+namespace UGSCore
+{
+
 enum class EEventType
 {
 	Syncing,
@@ -179,3 +182,4 @@ private:
 	void UpdateActiveInvestigations();
 };
 
+} // namespace UGSCore
diff --git a/Engine/Source/Programs/SlateUGS/Private/UGSCore/FileFilter.cpp b/Engine/Source/Programs/SlateUGS/Private/UGSCore/FileFilter.cpp
--- a/Engine/Source/Programs/SlateUGS/Private/UGSCore/FileFilter.cpp
+++ b/Engine/Source/Programs/SlateUGS/Private/UGSCore/FileFilter.cpp
@@ -2,6 +2,9 @@
 
 #include "FileFilter.h"
 
+namespace UGSCore
+{
+
 /// <summary>
 /// Node within a filter tree. Each node matches a single path fragment - a folder or file name, with an include or exclude rule consisting of a sequence of nodes. 
 /// </summary>
@@ -394,3 +397,5 @@ int FFileFilter::HighestPossibleIncludeMatch(const FNode* CurrentNode, const TAr
 	}
 	return BestRuleNumber;
 }
+
+} // namespace UGSCore
diff --git a/Engine/Source/Programs/SlateUGS/Private/UGSCore/FileFilter.h b/Engine/Source/Programs/SlateUGS/Private/UGSCore/FileFilter.h
--- a/Engine/Source/Programs/SlateUGS/Private/UGSCore/FileFilter.h
+++ b/Engine/Source/Programs/SlateUGS/Private/UGSCore/FileFilter.h
@@ -4,6 +4,9 @@
 
 #include "CoreMinimal.h"
 
+namespace UGSCore
+{
+
 struct FFileFilterNode;
 
 /// <summary>
@@ -119,3 +122,5 @@ private:
 	/// <returns>New highest rule number</returns>
 	static int HighestPossibleIncludeMatch(const FNode* CurrentNode, const TArray<FString>& Tokens, int TokenIdx, int CurrentBestRuleNumber);
 };
+
+} // namespace UGSCore
diff --git a/Engine/Source/Programs/SlateUGS/Private/UGSCore/ModalTask.cpp b/Engine/Source/Programs/SlateUGS/Private/UGSCore/ModalTask.cpp
--- a/Engine/Source/Programs/SlateUGS/Private/UGSCore/ModalTask.cpp
+++ b/Engine/Source/Programs/SlateUGS/Private/UGSCore/ModalTask.cpp
@@ -2,6 +2,9 @@
 
 #include "ModalTask.h"
 
+namespace UGSCore
+{
+
 //// FModalTaskResult ////
 
 FModalTaskResult::FModalTaskResult(bool bInSucceeded, const FText& InMessage)
@@ -49,3 +52,5 @@ TSharedRef<FModalTaskResult> FModalTaskResult::Aborted()
 IModalTask::~IModalTask()
 {
 }
+
+} // namespace UGSCore
diff --git a/Engine/Source/Programs/SlateUGS/Private/UGSCore/ModalTask.h b/Engine/Source/Programs/SlateUGS/Private/UGSCore/ModalTask.h
--- a/Engine/Source/Programs/SlateUGS/Private/UGSCore/ModalTask.h
+++ b/Engine/Source/Programs/SlateUGS/Private/UGSCore/ModalTask.h
@@ -5,6 +5,9 @@
 #include "CoreMinimal.h"
 #include "Templates/SharedPointer.h"
 
+namespace UGSCore
+{
+
 class FModalTaskResult
 {
 public:
@@ -30,3 +33,5 @@ public:
 	virtual ~IModalTask();
 	virtual TSharedRef<FModalTaskResult> Run(FEvent* AbortEvent) = 0;
 };
+
+} // namespace UGSCore
diff --git a/Engine/Source/Programs/SlateUGS/Private/UGSCore/OutputAdapters.cpp b/Engine/Source/Programs/SlateUGS/Private/UGSCore/OutputAdapters.cpp
--- a/Engine/Source/Programs/SlateUGS/Private/UGSCore/OutputAdapters.cpp
+++ b/Engine/Source/Programs/SlateUGS/Private/UGSCore/OutputAdapters.cpp
@@ -3,6 +3,9 @@
 #include "OutputAdapters.h"
 #include "Utility.h"
 
+namespace UGSCore
+{
+
 //// FLineBasedTextWriter ////
 
 FLineBasedTextWriter::FLineBasedTextWriter()
@@ -433,3 +436,5 @@ bool FProgressTextWriter::ReadFraction(const TArray<FString>& Tokens, int& Token
 	Fraction = 0.0f;
 	return false;
 }
+
+} // namespace UGSCore
diff --git a/Engine/Source/Programs/SlateUGS/Private/UGSCore/OutputAdapters.h b/Engine/Source/Programs/SlateUGS/Private/UGSCore/OutputAdapters.h
--- a/Engine/Source/Programs/SlateUGS/Private/UGSCore/OutputAdapters.h
+++ b/Engine/Source/Programs/SlateUGS/Private/UGSCore/OutputAdapters.h
@@ -4,6 +4,9 @@
 
 #include "CoreMinimal.h"
 
+namespace UGSCore
+{
+
 class FLineBasedTextWriter : public FOutputDevice
 {
 public:
@@ -120,3 +123,5 @@ private:
 	static TArray<FString> ParseTokens(const FString& Line);
 	static bool ReadFraction(const TArray<FString>& Tokens, int& TokenIdx, float& Fraction);
 };
+
+} // namespace UGSCore
diff --git a/Engine/Source/Programs/SlateUGS/Private/UGSCore/Perforce.cpp b/Engine/Source/Programs/SlateUGS/Private/UGSCore/Perforce.cpp
--- a/Engine/Source/Programs/SlateUGS/Private/UGSCore/Perforce.cpp
+++ b/Engine/Source/Programs/SlateUGS/Private/UGSCore/Perforce.cpp
@@ -9,6 +9,9 @@
 #include "HAL/Event.h"
 #include "Utility.h"
 
+namespace UGSCore
+{
+
 bool TryGetValue(const TMap<FString, FString>& Map, const TCHAR* Key, FString& OutValue)
 {
 	const FString* Value = Map.Find(Key);
@@ -1510,3 +1513,5 @@ bool FPerforceUtils::TryParseDateTime(const FString& Date, const FString& Time,
 	OutDate = FDateTime(Year, Month, Day, Hour, Minute, Second);
 	return true;
 }
+
+} // namespace UGSCore
diff --git a/Engine/Source/Programs/SlateUGS/Private/UGSCore/Perforce.h b/Engine/Source/Programs/SlateUGS/Private/UGSCore/Perforce.h
--- a/Engine/Source/Programs/SlateUGS/Private/UGSCore/Perforce.h
+++ b/Engine/Source/Programs/SlateUGS/Private/UGSCore/Perforce.h
@@ -7,6 +7,9 @@
 #include "Misc/EnumClassFlags.h"
 #include "Templates/Tuple.h"
 
+namespace UGSCore
+{
+
 struct FPerforceChangeSummary
 {
 	int Number;
@@ -270,3 +273,5 @@ struct FPerforceUtils
 	static FString GetClientOrDepotDirectoryName(const TCHAR* ClientFile);
 	static bool TryParseDateTime(const FString& Date, const FString& Time, FDateTime& OutDate);
 };
+
+} // namespace UGSCore
diff --git a/Engine/Source/Programs/SlateUGS/Private/UGSCore/PerforceMonitor.cpp b/Engine/Source/Programs/SlateUGS/Private/UGSCore/PerforceMonitor.cpp
--- a/Engine/Source/Programs/SlateUGS/Private/UGSCore/PerforceMonitor.cpp
+++ b/Engine/Source/Programs/SlateUGS/Private/UGSCore/PerforceMonitor.cpp
@@ -8,6 +8,9 @@
 #include "CustomConfigFile.h"
 #include "Utility.h"
 
+namespace UGSCore
+{
+
 FPerforceMonitor::FPerforceMonitor(const TSharedRef<FPerforceConnection>& InPerforce, const FString& InBranchClientPath, const FString& InSelectedClientFileName, const FString& InSelectedProjectIdentifier, const FString& InLogPath)
 	: Perforce(InPerforce)
 	, BranchClientPath(InBranchClientPath)
@@ -507,3 +510,5 @@ bool FPerforceMonitor::UpdateZippedBinaries()
 	return true;
 }
 
+
+} // namespace UGSCore
diff --git a/Engine/Source/Programs/SlateUGS/Private/UGSCore/PerforceMonitor.h b/Engine/Source/Programs/SlateUGS/Private/UGSCore/PerforceMonitor.h
--- a/Engine/Source/Programs/SlateUGS/Private/UGSCore/PerforceMonitor.h
+++ b/Engine/Source/Programs/SlateUGS/Private/UGSCore/PerforceMonitor.h
@@ -9,6 +9,9 @@
 #include "Perforce.h"
 #include "OutputAdapters.h"
 
+namespace UGSCore
+{
+
 enum class EPerforceChangeType
 {
 	Code,
@@ -75,3 +78,5 @@ private:
 	bool UpdateChangeTypes();
 	bool UpdateZippedBinaries();
 };
+
+} // namespace UGSCore
diff --git a/Engine/Source/Programs/SlateUGS/Private/UGSCore/Telemetry.cpp b/Engine/Source/Programs/SlateUGS/Private/UGSCore/Telemetry.cpp
--- a/Engine/Source/Programs/SlateUGS/Private/UGSCore/Telemetry.cpp
+++ b/Engine/Source/Programs/SlateUGS/Private/UGSCore/Telemetry.cpp
@@ -2,6 +2,9 @@
 
 #include "Telemetry.h"
 
+namespace UGSCore
+{
+
 //// FTelemetryStopwatch ////
 
 FTelemetryStopwatch::FTelemetryStopwatch(const FString& InAction, const FString& InProject)
@@ -50,3 +53,4 @@ FTelemetryWriter::~FTelemetryWriter()
 {
 }
 
+} // namespace UGSCore
diff --git a/Engine/Source/Programs/SlateUGS/Private/UGSCore/Telemetry.h b/Engine/Source/Programs/SlateUGS/Private/UGSCore/Telemetry.h
--- a/Engine/Source/Programs/SlateUGS/Private/UGSCore/Telemetry.h
+++ b/Engine/Source/Programs/SlateUGS/Private/UGSCore/Telemetry.h
@@ -5,6 +5,9 @@
 #include "CoreMinimal.h"
 #include "Misc/Timespan.h"
 
+namespace UGSCore
+{
+
 class FTelemetryStopwatch
 {
 public:
@@ -28,3 +31,5 @@ public:
 	FTelemetryWriter(const FString& InSqlConnectionString, const FString& InLogFileName);
 	~FTelemetryWriter();
 };
+
+} // namespace UGSCore
diff --git a/Engine/Source/Programs/SlateUGS/Private/UGSCore/UserSettings.cpp b/Engine/Source/Programs/SlateUGS/Private/UGSCore/UserSettings.cpp
--- a/Engine/Source/Programs/SlateUGS/Private/UGSCore/UserSettings.cpp
+++ b/Engine/Source/Programs/SlateUGS/Private/UGSCore/UserSettings.cpp
@@ -5,6 +5,9 @@
 #include "Algo/Transform.h"
 #include "Utility.h"
 
+namespace UGSCore
+{
+
 //// EBuildConfig ////
 
 FString ToString(EBuildConfig BuildConfig)
@@ -612,3 +615,5 @@ FString FUserSettings::UnescapeText(const FString& Text)
 	}
 	return Result;
 }
+
+} // namespace UGSCore
diff --git a/Engine/Source/Programs/SlateUGS/Private/UGSCore/UserSettings.h b/Engine/Source/Programs/SlateUGS/Private/UGSCore/UserSettings.h
--- a/Engine/Source/Programs/SlateUGS/Private/UGSCore/UserSettings.h
+++ b/Engine/Source/Programs/SlateUGS/Private/UGSCore/UserSettings.h
@@ -8,6 +8,9 @@
 #include "Misc/Guid.h"
 #include "Workspace.h"
 
+namespace UGSCore
+{
+
 enum class EBuildConfig
 {
 	Debug,
@@ -126,3 +129,5 @@ struct FUserSettings
 	static FString EscapeText(const FString& Text);
 	static FString UnescapeText(const FString& Text);
 };
+
+} // namespace UGSCore
diff --git a/Engine/Source/Programs/SlateUGS/Private/UGSCore/Utility.cpp b/Engine/Source/Programs/SlateUGS/Private/UGSCore/Utility.cpp
--- a/Engine/Source/Programs/SlateUGS/Private/UGSCore/Utility.cpp
+++ b/Engine/Source/Programs/SlateUGS/Private/UGSCore/Utility.cpp
@@ -4,10 +4,12 @@
 #include "HAL/FileManager.h"
 #include "HAL/PlatformFileManager.h"
 #include "HAL/PlatformProcess.h"
-//#include "ChildProcess.h"
 #include "Misc/Paths.h"
 #include "HAL/Event.h"
 
+namespace UGSCore
+{
+
 bool FUtility::TryParse(const TCHAR* Text, int32& OutValue)
 {
 	TCHAR* TextEnd;
@@ -240,3 +242,5 @@ FString FUtility::ExpandVariables(const TCHAR* InputString, const TMap<FString,
 	}
 	return Result;
 }
+
+} // namespace UGSCore
diff --git a/Engine/Source/Programs/SlateUGS/Private/UGSCore/Utility.h b/Engine/Source/Programs/SlateUGS/Private/UGSCore/Utility.h
--- a/Engine/Source/Programs/SlateUGS/Private/UGSCore/Utility.h
+++ b/Engine/Source/Programs/SlateUGS/Private/UGSCore/Utility.h
@@ -5,6 +5,9 @@
 #include "CoreMinimal.h"
 #include "OutputAdapters.h"
 
+namespace UGSCore
+{
+
 struct FAbortException
 {
 };
@@ -33,3 +36,5 @@ struct FUtility
 	 */
 	static FString ExpandVariables(const TCHAR* InputString, const TMap<FString, FString>* AdditionalVariables = nullptr);
 };
+
+} // namespace UGSCore
diff --git a/Engine/Source/Programs/SlateUGS/Private/UGSCore/Workspace.cpp b/Engine/Source/Programs/SlateUGS/Private/UGSCore/Workspace.cpp
--- a/Engine/Source/Programs/SlateUGS/Private/UGSCore/Workspace.cpp
+++ b/Engine/Source/Programs/SlateUGS/Private/UGSCore/Workspace.cpp
@@ -13,6 +13,9 @@
 #include "Telemetry.h"
 #include "Utility.h"
 
+namespace UGSCore
+{
+
 const TCHAR* FWorkspace::DefaultBuildTargets[] =
 {
 	TEXT("UnrealHeaderTool Win64 Development, 0.1"),
@@ -1260,3 +1263,4 @@ bool FWorkspace::ReadToken(const FString& Line, int32& LineIdx, FString &OutToke
 	return true;
 }
 
+} // namespace UGSCore
diff --git a/Engine/Source/Programs/SlateUGS/Private/UGSCore/Workspace.h b/Engine/Source/Programs/SlateUGS/Private/UGSCore/Workspace.h
--- a/Engine/Source/Programs/SlateUGS/Private/UGSCore/Workspace.h
+++ b/Engine/Source/Programs/SlateUGS/Private/UGSCore/Workspace.h
@@ -10,6 +10,9 @@
 #include "HAL/Runnable.h"
 #include "Misc/EnumClassFlags.h"
 
+namespace UGSCore
+{
+
 enum class EWorkspaceUpdateOptions
 {
 	Sync = 0x01,
@@ -146,3 +149,5 @@ private:
 	static bool UpdateVersionLine(FString& Line, const FString& Prefix, const FString& Suffix);
 	static bool ReadToken(const FString& Line, int32& LineIdx, FString &OutToken);
 };
+
+} // namespace UGSCore
diff --git a/Engine/Source/Programs/SlateUGS/Private/UGSTab.cpp b/Engine/Source/Programs/SlateUGS/Private/UGSTab.cpp
--- a/Engine/Source/Programs/SlateUGS/Private/UGSTab.cpp
+++ b/Engine/Source/Programs/SlateUGS/Private/UGSTab.cpp
@@ -66,7 +66,7 @@ namespace
 		const TCHAR* HostPlatform = TEXT("Linux");
 	#endif
 
-	class FLineWriter : public FLineBasedTextWriter
+	class FLineWriter : public UGSCore::FLineBasedTextWriter
 	{
 	public:
 		virtual void FlushLine(const FString& Line) override
@@ -85,25 +85,25 @@ namespace
 
 	// Honestly ... seems ... super hacky/hardcoded. With out all these you assert when trying to merge build targets sooo a bit odd
 	// TODO Need to do each of these ... per ... platform??
-	TMap<FGuid, FCustomConfigObject> GetDefaultBuildStepObjects(const FString& EditorTargetName, TSharedPtr<FUserSettings> Settings)
+	TMap<FGuid, UGSCore::FCustomConfigObject> GetDefaultBuildStepObjects(const FString& EditorTargetName, TSharedPtr<UGSCore::FUserSettings> Settings)
 	{
-		TArray<FBuildStep> DefaultBuildSteps;
-		DefaultBuildSteps.Add(FBuildStep(FGuid(0x01F66060, 0x73FA4CC8, 0x9CB3E217, 0xFBBA954E), 0, TEXT("Compile UnrealHeaderTool"), TEXT("Compiling UnrealHeaderTool..."), 1, TEXT("UnrealHeaderTool"), HostPlatform, TEXT("Development"), TEXT(""), !ShouldSyncPrecompiledEditor()));
+		TArray<UGSCore::FBuildStep> DefaultBuildSteps;
+		DefaultBuildSteps.Add(UGSCore::FBuildStep(FGuid(0x01F66060, 0x73FA4CC8, 0x9CB3E217, 0xFBBA954E), 0, TEXT("Compile UnrealHeaderTool"), TEXT("Compiling UnrealHeaderTool..."), 1, TEXT("UnrealHeaderTool"), HostPlatform, TEXT("Development"), TEXT(""), !ShouldSyncPrecompiledEditor()));
 		FString ActualEditorTargetName = (EditorTargetName.Len() > 0) ? EditorTargetName : "UnrealEditor";
-		DefaultBuildSteps.Add(FBuildStep(FGuid(0xF097FF61, 0xC9164058, 0x839135B4, 0x6C3173D5), 1, FString::Printf(TEXT("Compile %s"), *ActualEditorTargetName), FString::Printf(TEXT("Compiling %s..."), *ActualEditorTargetName), 10, ActualEditorTargetName, HostPlatform, ::ToString(Settings->CompiledEditorBuildConfig), TEXT(""), !ShouldSyncPrecompiledEditor()));
-		DefaultBuildSteps.Add(FBuildStep(FGuid(0xC6E633A1, 0x956F4AD3, 0xBC956D06, 0xD131E7B4), 2, TEXT("Compile ShaderCompileWorker"), TEXT("Compiling ShaderCompileWorker..."), 1, TEXT("ShaderCompileWorker"), HostPlatform, TEXT("Development"), TEXT(""), !ShouldSyncPrecompiledEditor()));
-		DefaultBuildSteps.Add(FBuildStep(FGuid(0x24FFD88C, 0x79014899, 0x9696AE10, 0x66B4B6E8), 3, TEXT("Compile UnrealLightmass"), TEXT("Compiling UnrealLightmass..."), 1, TEXT("UnrealLightmass"), HostPlatform, TEXT("Development"), TEXT(""), !ShouldSyncPrecompiledEditor()));
-		DefaultBuildSteps.Add(FBuildStep(FGuid(0xFFF20379, 0x06BF4205, 0x8A3EC534, 0x27736688), 4, TEXT("Compile CrashReportClient"), TEXT("Compiling CrashReportClient..."), 1, TEXT("CrashReportClient"), HostPlatform, TEXT("Development"), TEXT(""), !ShouldSyncPrecompiledEditor()));
-		DefaultBuildSteps.Add(FBuildStep(FGuid(0x89FE8A79, 0xD2594C7B, 0xBFB468F7, 0x218B91C2), 5, TEXT("Compile UnrealInsights"), TEXT("Compiling UnrealInsights..."), 1, TEXT("UnrealInsights"), HostPlatform, TEXT("Development"), TEXT(""), !ShouldSyncPrecompiledEditor()));
-		DefaultBuildSteps.Add(FBuildStep(FGuid(0x46312669, 0x5069428D, 0x8D72C241, 0x6C5A322E), 6, TEXT("Launch UnrealInsights"), TEXT("Running UnrealInsights..."), 1, TEXT("UnrealInsights"), HostPlatform, TEXT("Shipping"), TEXT(""), !ShouldSyncPrecompiledEditor()));
-		DefaultBuildSteps.Add(FBuildStep(FGuid(0xBB48CA5B, 0x56824432, 0x824DC451, 0x336A6523), 7, TEXT("Compile Zen Dashboard"), TEXT("Compile ZenDashboard Step..."), 1, TEXT("ZenDashboard"), HostPlatform, TEXT("Development"), TEXT(""), !ShouldSyncPrecompiledEditor()));
-		DefaultBuildSteps.Add(FBuildStep(FGuid(0x586CC0D3, 0x39144DF9, 0xACB62C02, 0xCD9D4FC6), 8, TEXT("Launch Zen Dashboard"), TEXT("Running Zen Dashboard..."), 1, TEXT("ZenDashboard"), HostPlatform, TEXT("Development"), TEXT(""), !ShouldSyncPrecompiledEditor()));
-		DefaultBuildSteps.Add(FBuildStep(FGuid(0x91C2A429, 0xC39149B4, 0x92A54E6B, 0xE71E0F00), 9, TEXT("Compile SwitchboardListener"), TEXT("Compiling SwitchboardListener..."), 1, TEXT("SwitchboardListener"), HostPlatform, TEXT("Development"), TEXT(""), !ShouldSyncPrecompiledEditor()));
-		DefaultBuildSteps.Add(FBuildStep(FGuid(0x5036C75B, 0x8DF04329, 0x82A1869D, 0xD2D48605), 10, TEXT("Compile UnrealMultiUserServer"), TEXT("Compiling UnrealMultiUserServer..."), 1, TEXT("UnrealMultiUserServer"), HostPlatform, TEXT("Development"), TEXT(""), !ShouldSyncPrecompiledEditor()));
-		DefaultBuildSteps.Add(FBuildStep(FGuid(0x274B89C3, 0x9DC64465, 0xA50840AB, 0xC4593CC2), 11, TEXT("Compile UnrealMultiUserSlateServer"), TEXT("Compiling UnrealMultiUserSlateServer..."), 1, TEXT("UnrealMultiUserSlateServer"), HostPlatform, TEXT("Development"), TEXT(""), !ShouldSyncPrecompiledEditor()));
-
-		TMap<FGuid,FCustomConfigObject> DefaultBuildStepObjects;
-		for (const FBuildStep& DefaultBuildStep : DefaultBuildSteps)
+		DefaultBuildSteps.Add(UGSCore::FBuildStep(FGuid(0xF097FF61, 0xC9164058, 0x839135B4, 0x6C3173D5), 1, FString::Printf(TEXT("Compile %s"), *ActualEditorTargetName), FString::Printf(TEXT("Compiling %s..."), *ActualEditorTargetName), 10, ActualEditorTargetName, HostPlatform, UGSCore::ToString(Settings->CompiledEditorBuildConfig), TEXT(""), !ShouldSyncPrecompiledEditor()));
+		DefaultBuildSteps.Add(UGSCore::FBuildStep(FGuid(0xC6E633A1, 0x956F4AD3, 0xBC956D06, 0xD131E7B4), 2, TEXT("Compile ShaderCompileWorker"), TEXT("Compiling ShaderCompileWorker..."), 1, TEXT("ShaderCompileWorker"), HostPlatform, TEXT("Development"), TEXT(""), !ShouldSyncPrecompiledEditor()));
+		DefaultBuildSteps.Add(UGSCore::FBuildStep(FGuid(0x24FFD88C, 0x79014899, 0x9696AE10, 0x66B4B6E8), 3, TEXT("Compile UnrealLightmass"), TEXT("Compiling UnrealLightmass..."), 1, TEXT("UnrealLightmass"), HostPlatform, TEXT("Development"), TEXT(""), !ShouldSyncPrecompiledEditor()));
+		DefaultBuildSteps.Add(UGSCore::FBuildStep(FGuid(0xFFF20379, 0x06BF4205, 0x8A3EC534, 0x27736688), 4, TEXT("Compile CrashReportClient"), TEXT("Compiling CrashReportClient..."), 1, TEXT("CrashReportClient"), HostPlatform, TEXT("Development"), TEXT(""), !ShouldSyncPrecompiledEditor()));
+		DefaultBuildSteps.Add(UGSCore::FBuildStep(FGuid(0x89FE8A79, 0xD2594C7B, 0xBFB468F7, 0x218B91C2), 5, TEXT("Compile UnrealInsights"), TEXT("Compiling UnrealInsights..."), 1, TEXT("UnrealInsights"), HostPlatform, TEXT("Development"), TEXT(""), !ShouldSyncPrecompiledEditor()));
+		DefaultBuildSteps.Add(UGSCore::FBuildStep(FGuid(0x46312669, 0x5069428D, 0x8D72C241, 0x6C5A322E), 6, TEXT("Launch UnrealInsights"), TEXT("Running UnrealInsights..."), 1, TEXT("UnrealInsights"), HostPlatform, TEXT("Shipping"), TEXT(""), !ShouldSyncPrecompiledEditor()));
+		DefaultBuildSteps.Add(UGSCore::FBuildStep(FGuid(0xBB48CA5B, 0x56824432, 0x824DC451, 0x336A6523), 7, TEXT("Compile Zen Dashboard"), TEXT("Compile ZenDashboard Step..."), 1, TEXT("ZenDashboard"), HostPlatform, TEXT("Development"), TEXT(""), !ShouldSyncPrecompiledEditor()));
+		DefaultBuildSteps.Add(UGSCore::FBuildStep(FGuid(0x586CC0D3, 0x39144DF9, 0xACB62C02, 0xCD9D4FC6), 8, TEXT("Launch Zen Dashboard"), TEXT("Running Zen Dashboard..."), 1, TEXT("ZenDashboard"), HostPlatform, TEXT("Development"), TEXT(""), !ShouldSyncPrecompiledEditor()));
+		DefaultBuildSteps.Add(UGSCore::FBuildStep(FGuid(0x91C2A429, 0xC39149B4, 0x92A54E6B, 0xE71E0F00), 9, TEXT("Compile SwitchboardListener"), TEXT("Compiling SwitchboardListener..."), 1, TEXT("SwitchboardListener"), HostPlatform, TEXT("Development"), TEXT(""), !ShouldSyncPrecompiledEditor()));
+		DefaultBuildSteps.Add(UGSCore::FBuildStep(FGuid(0x5036C75B, 0x8DF04329, 0x82A1869D, 0xD2D48605), 10, TEXT("Compile UnrealMultiUserServer"), TEXT("Compiling UnrealMultiUserServer..."), 1, TEXT("UnrealMultiUserServer"), HostPlatform, TEXT("Development"), TEXT(""), !ShouldSyncPrecompiledEditor()));
+		DefaultBuildSteps.Add(UGSCore::FBuildStep(FGuid(0x274B89C3, 0x9DC64465, 0xA50840AB, 0xC4593CC2), 11, TEXT("Compile UnrealMultiUserSlateServer"), TEXT("Compiling UnrealMultiUserSlateServer..."), 1, TEXT("UnrealMultiUserSlateServer"), HostPlatform, TEXT("Development"), TEXT(""), !ShouldSyncPrecompiledEditor()));
+
+		TMap<FGuid, UGSCore::FCustomConfigObject> DefaultBuildStepObjects;
+		for (const UGSCore::FBuildStep& DefaultBuildStep : DefaultBuildSteps)
 		{
 			DefaultBuildStepObjects.Add(DefaultBuildStep.UniqueId, DefaultBuildStep.ToConfigObject());
 		}
@@ -111,30 +111,30 @@ namespace
 		return DefaultBuildStepObjects;
 	}
 
-	EBuildConfig GetEditorBuildConfig()
+	UGSCore::EBuildConfig GetEditorBuildConfig()
 	{
-		return EBuildConfig::Development;
-		//return ShouldSyncPrecompiledEditor()? EBuildConfig::Development : Settings->CompiledEditorBuildConfig;
+		return UGSCore::EBuildConfig::Development;
+		//return ShouldSyncPrecompiledEditor() ? UGSCore::EBuildConfig::Development : Settings->CompiledEditorBuildConfig;
 	}
 
-	FString GetEditorExePath(EBuildConfig Config, TSharedPtr<FDetectProjectSettingsTask> DetectSettings)
+	FString GetEditorExePath(UGSCore::EBuildConfig Config, TSharedPtr<UGSCore::FDetectProjectSettingsTask> DetectSettings)
 	{
 		FString ExeFileName = TEXT("UnrealEditor");
 
-		if (Config != EBuildConfig::DebugGame && Config != EBuildConfig::Development)
+		if (Config != UGSCore::EBuildConfig::DebugGame && Config != UGSCore::EBuildConfig::Development)
 		{
 	#if PLATFORM_WINDOWS
-			ExeFileName = FString::Printf(TEXT("UnrealEditor-%s-%s.exe"), HostPlatform, *::ToString(Config));
+			ExeFileName = FString::Printf(TEXT("UnrealEditor-%s-%s.exe"), HostPlatform, *UGSCore::ToString(Config));
 	#else
-			ExeFileName = FString::Printf(TEXT("UnrealEditor-%s-%s"), HostPlatform, *::ToString(Config));
+			ExeFileName = FString::Printf(TEXT("UnrealEditor-%s-%s"), HostPlatform, *UGSCore::ToString(Config));
 	#endif
 		}
 		return DetectSettings->BranchDirectoryName / "Engine" / "Binaries" / HostPlatform / ExeFileName;
 	}
 
-	TMap<FString, FString> GetWorkspaceVariables(TSharedPtr<FDetectProjectSettingsTask> DetectSettings)
+	TMap<FString, FString> GetWorkspaceVariables(TSharedPtr<UGSCore::FDetectProjectSettingsTask> DetectSettings)
 	{
-		EBuildConfig EditorBuildConfig = GetEditorBuildConfig();
+		UGSCore::EBuildConfig EditorBuildConfig = GetEditorBuildConfig();
 
 		TMap<FString, FString> Variables;
 		Variables.Add("BranchDir", DetectSettings->BranchDirectoryName);
@@ -144,21 +144,22 @@ namespace
 		// Todo: These might not be called "UE4*" anymore
 		Variables.Add("UE4EditorExe", GetEditorExePath(EditorBuildConfig, DetectSettings));
 		Variables.Add("UE4EditorCmdExe", GetEditorExePath(EditorBuildConfig, DetectSettings).Replace(TEXT(".exe"), TEXT("-Cmd.exe")));
-		Variables.Add("UE4EditorConfig", ::ToString(EditorBuildConfig));
-		Variables.Add("UE4EditorDebugArg", (EditorBuildConfig == EBuildConfig::Debug || EditorBuildConfig == EBuildConfig::DebugGame)? " -debug" : "");
+		Variables.Add("UE4EditorConfig", UGSCore::ToString(EditorBuildConfig));
+		Variables.Add("UE4EditorDebugArg", (EditorBuildConfig == UGSCore::EBuildConfig::Debug || EditorBuildConfig == UGSCore::EBuildConfig::DebugGame)? " -debug" : "");
+
 		return Variables;
 	}
 
 	static FString FormatUserName(FString UserName)
 	{
 		TStringBuilder<256> NormalUserName;
-		for(int Index = 0; Index < UserName.Len(); Index++)
+		for (int Index = 0; Index < UserName.Len(); Index++)
 		{
-			if(Index == 0 || UserName[Index - 1] == '.')
+			if (Index == 0 || UserName[Index - 1] == '.')
 			{
 				NormalUserName.AppendChar(FChar::ToUpper(UserName[Index]));
 			}
-			else if(UserName[Index] == '.')
+			else if (UserName[Index] == '.')
 			{
 				NormalUserName.AppendChar(' ');
 			}
@@ -167,6 +168,7 @@ namespace
 				NormalUserName.AppendChar(FChar::ToLower(UserName[Index]));
 			}
 		}
+
 		return NormalUserName.ToString();
 	}
 }
@@ -193,7 +195,7 @@ bool UGSTab::OnWorkspaceChosen(const FString& Project)
 
 void UGSTab::OnSyncChangelist(int Changelist)
 {
-	TSharedRef<FWorkspaceUpdateContext, ESPMode::ThreadSafe> Context = MakeShared<FWorkspaceUpdateContext, ESPMode::ThreadSafe>(
+	TSharedRef<UGSCore::FWorkspaceUpdateContext, ESPMode::ThreadSafe> Context = MakeShared<UGSCore::FWorkspaceUpdateContext, ESPMode::ThreadSafe>(
 		Changelist,
 		Options,
 		CombinedSyncFilter,
@@ -260,7 +262,7 @@ void UGSTab::Tick()
 
 namespace
 {
-	bool ShouldShowChange(const FPerforceChangeSummary& Change, const TArray<FString>& ExcludeChanges)
+	bool ShouldShowChange(const UGSCore::FPerforceChangeSummary& Change, const TArray<FString>& ExcludeChanges)
 	{
 		for (const FString& ExcludeChange : ExcludeChanges)
 		{
@@ -289,15 +291,15 @@ bool UGSTab::ShouldIncludeInReviewedList(const TSet<int>& PromotedChangeNumbers,
 		return true;
 	}
 
-	TSharedPtr<FEventSummary> Review;
+	TSharedPtr<UGSCore::FEventSummary> Review;
 	if (EventMonitor->TryGetSummaryForChange(ChangeNumber, Review))
 	{
-		if (Review->LastStarReview.IsValid() && Review->LastStarReview->Type == EEventType::Starred)
+		if (Review->LastStarReview.IsValid() && Review->LastStarReview->Type == UGSCore::EEventType::Starred)
 		{
 			return true;
 		}
 
-		if (Review->Verdict == EReviewVerdict::Good || Review->Verdict == EReviewVerdict::Mixed)
+		if (Review->Verdict == UGSCore::EReviewVerdict::Good || Review->Verdict == UGSCore::EReviewVerdict::Mixed)
 		{
 			return true;
 		}
@@ -308,14 +310,14 @@ bool UGSTab::ShouldIncludeInReviewedList(const TSet<int>& PromotedChangeNumbers,
 
 void UGSTab::UpdateGameTabBuildList()
 {
-	TArray<TSharedRef<FPerforceChangeSummary, ESPMode::ThreadSafe>> Changes = PerforceMonitor->GetChanges();
+	TArray<TSharedRef<UGSCore::FPerforceChangeSummary, ESPMode::ThreadSafe>> Changes = PerforceMonitor->GetChanges();
 	TArray<TSharedPtr<FChangeInfo>> ChangeInfos;
 
 	// do we need to store these for something?
 	TSet<int> PromotedChangeNumbers = PerforceMonitor->GetPromotedChangeNumbers();
 
 	TArray<FString> ExcludeChanges;
-	TSharedPtr<FCustomConfigFile, ESPMode::ThreadSafe> ProjectConfigFile = Workspace->GetProjectConfigFile();
+	TSharedPtr<UGSCore::FCustomConfigFile, ESPMode::ThreadSafe> ProjectConfigFile = Workspace->GetProjectConfigFile();
 	if (ProjectConfigFile.IsValid())
 	{
 		ProjectConfigFile->TryGetValues(TEXT("Options.ExcludeChanges"), ExcludeChanges);
@@ -327,7 +329,7 @@ void UGSTab::UpdateGameTabBuildList()
 
 	for (int ChangeIdx = 0; ChangeIdx < Changes.Num(); ChangeIdx++)
 	{
-		const FPerforceChangeSummary& Change = Changes[ChangeIdx].Get();
+		const UGSCore::FPerforceChangeSummary& Change = Changes[ChangeIdx].Get();
 		if (ShouldShowChange(Change, ExcludeChanges) || PromotedChangeNumbers.Contains(Change.Number))
 		{
 			if (!bOnlyShowReviewed || (!EventMonitor->IsUnderInvestigation(Change.Number) && (ShouldIncludeInReviewedList(PromotedChangeNumbers, Change.Number) || bFirstChange)))
@@ -340,13 +342,13 @@ void UGSTab::UpdateGameTabBuildList()
 					DisplayTime = (DisplayTime - DetectSettings->ServerTimeZone).GetDate();
 				}
 
-				TSharedPtr<FEventSummary> Review;
-				EReviewVerdict Status = EReviewVerdict::Unknown;
+				TSharedPtr<UGSCore::FEventSummary> Review;
+				UGSCore::EReviewVerdict Status = UGSCore::EReviewVerdict::Unknown;
 				if (EventMonitor->TryGetSummaryForChange(Change.Number, Review) && Review->LastStarReview.IsValid())
 				{
-					if (Review->LastStarReview->Type == EEventType::Starred)
+					if (Review->LastStarReview->Type == UGSCore::EEventType::Starred)
 					{
-						Status = EReviewVerdict::Good;
+						Status = UGSCore::EReviewVerdict::Good;
 					}
 					else
 					{
@@ -395,10 +397,10 @@ FString UGSTab::GetSyncProgress() const
 	return Workspace->GetCurrentProgress().Key;
 }
 
-TArray<FWorkspaceSyncCategory> UGSTab::GetSyncCategories(SyncCategoryType CategoryType) const
+TArray<UGSCore::FWorkspaceSyncCategory> UGSTab::GetSyncCategories(SyncCategoryType CategoryType) const
 {
-	TArray<FWorkspaceSyncCategory> SyncCategories;
-	TMap<FGuid, FWorkspaceSyncCategory> Categories = Workspace->GetSyncCategories();
+	TArray<UGSCore::FWorkspaceSyncCategory> SyncCategories;
+	TMap<FGuid, UGSCore::FWorkspaceSyncCategory> Categories = Workspace->GetSyncCategories();
 	TArray<FGuid> SyncExcludedCategories;
 
 	if (CategoryType == SyncCategoryType::CurrentWorkspace)
@@ -414,7 +416,9 @@ TArray<FWorkspaceSyncCategory> UGSTab::GetSyncCategories(SyncCategoryType Catego
 	{
 		Categories[Guid].bEnable = false;
 	}
+
 	Categories.GenerateValueArray(SyncCategories);
+
 	return SyncCategories;
 }
 
@@ -424,6 +428,7 @@ TArray<FString> UGSTab::GetSyncViews(SyncCategoryType CategoryType) const
 	{
 		return WorkspaceSettings->SyncView;
 	}
+
 	return UserSettings->SyncView;
 }
 
@@ -438,7 +443,7 @@ UGSTabManager* UGSTab::GetTabManager()
 }
 
 // This is getting called on a thread, lock our stuff up
-void UGSTab::OnWorkspaceSyncComplete(TSharedRef<FWorkspaceUpdateContext, ESPMode::ThreadSafe> WorkspaceContext, EWorkspaceUpdateResult SyncResult, const FString& StatusMessage)
+void UGSTab::OnWorkspaceSyncComplete(TSharedRef<UGSCore::FWorkspaceUpdateContext, ESPMode::ThreadSafe> WorkspaceContext, UGSCore::EWorkspaceUpdateResult SyncResult, const FString& StatusMessage)
 {
 	FScopeLock Lock(&CriticalSection);
 
@@ -460,16 +465,17 @@ void UGSTab::OnWorkspaceSyncComplete(TSharedRef<FWorkspaceUpdateContext, ESPMode
 
 void UGSTab::SetupWorkspace()
 {
-	ProjectFileName = FUtility::GetPathWithCorrectCase(ProjectFileName);
+	ProjectFileName = UGSCore::FUtility::GetPathWithCorrectCase(ProjectFileName);
 
 	// TODO likely should also log this on an Empty tab... so we can show logging info when we are loading things
-	DetectSettings = MakeShared<FDetectProjectSettingsTask>(MakeShared<FPerforceConnection>(TEXT(""), TEXT(""), TEXT("")), ProjectFileName, MakeShared<FLineWriter>());
+	DetectSettings = MakeShared<UGSCore::FDetectProjectSettingsTask>(MakeShared<UGSCore::FPerforceConnection>(TEXT(""), TEXT(""), TEXT("")), ProjectFileName, MakeShared<FLineWriter>());
 
-	TSharedRef<FModalTaskResult> Result = ExecuteModalTask(
+	TSharedRef<UGSCore::FModalTaskResult> Result = ExecuteModalTask(
 		FSlateApplication::Get().GetActiveModalWindow(),
 		DetectSettings.ToSharedRef(),
 		LOCTEXT("OpeningProjectTitle", "Opening Project"),
 		LOCTEXT("OpeningProjectCaption", "Opening project, please wait..."));
+
 	if (Result->Failed())
 	{
 		FMessageDialog::Open(EAppMsgType::Ok, Result->GetMessage());
@@ -479,11 +485,11 @@ void UGSTab::SetupWorkspace()
 	FString DataFolder = FString(FPlatformProcess::UserSettingsDir()) / TEXT("UnrealGameSync");
 	IFileManager::Get().MakeDirectory(*DataFolder);
 
-	UserSettings = MakeShared<FUserSettings>(*(DataFolder / TEXT("UnrealGameSync_Slate.ini")));
+	UserSettings = MakeShared<UGSCore::FUserSettings>(*(DataFolder / TEXT("UnrealGameSync_Slate.ini")));
 
-	PerforceClient = DetectSettings->PerforceClient;
+	PerforceClient    = DetectSettings->PerforceClient;
 	WorkspaceSettings = UserSettings->FindOrAddWorkspace(*DetectSettings->BranchClientPath);
-	ProjectSettings = UserSettings->FindOrAddProject(*DetectSettings->NewSelectedClientFileName);
+	ProjectSettings   = UserSettings->FindOrAddProject(*DetectSettings->NewSelectedClientFileName);
 
 	// Check if the project we've got open in this workspace is the one we're actually synced to
 	int CurrentChangeNumber = -1;
@@ -498,15 +504,15 @@ void UGSTab::SetupWorkspace()
 		ClientKey = ClientKey.Left(ClientKey.Len() - 1);
 	}
 
-	FString ProjectLogBaseName = DataFolder / FString::Printf(TEXT("%s@%s"), *PerforceClient->ClientName, *ClientKey.Replace(TEXT("/"), TEXT("$")));
-	FString TelemetryProjectIdentifier = FPerforceUtils::GetClientOrDepotDirectoryName(*DetectSettings->NewSelectedProjectIdentifier);
-	FString SelectedProjectIdentifier = DetectSettings->NewSelectedProjectIdentifier;
+	FString ProjectLogBaseName         = DataFolder / FString::Printf(TEXT("%s@%s"), *PerforceClient->ClientName, *ClientKey.Replace(TEXT("/"), TEXT("$")));
+	FString TelemetryProjectIdentifier = UGSCore::FPerforceUtils::GetClientOrDepotDirectoryName(*DetectSettings->NewSelectedProjectIdentifier);
+	FString SelectedProjectIdentifier  = DetectSettings->NewSelectedProjectIdentifier;
 
 	FString LogFileName = DataFolder / FPaths::GetPath(ProjectFileName) + TEXT(".sync.log");
 	GameSyncTabView->SetSyncLogLocation(LogFileName); // Todo: if SetSyncLogLocation fails, then it failed to create a log file and may need to handle that
 	GameSyncTabView->GetSyncLog()->AppendLine(TEXT("Creating log at: ") + LogFileName);
 
-	Workspace = MakeShared<FWorkspace>(
+	Workspace = MakeShared<UGSCore::FWorkspace>(
 		PerforceClient.ToSharedRef(),
 		DetectSettings->BranchDirectoryName,
 		ProjectFileName,
@@ -518,11 +524,11 @@ void UGSTab::SetupWorkspace()
 		TelemetryProjectIdentifier,
 		MakeShared<FLogWidgetTextWriter>(GameSyncTabView->GetSyncLog().ToSharedRef()));
 
-	Workspace->OnUpdateComplete = [this] (TSharedRef<FWorkspaceUpdateContext, ESPMode::ThreadSafe> WorkspaceContext, EWorkspaceUpdateResult SyncResult, const FString& StatusMessage) {
+	Workspace->OnUpdateComplete = [this] (TSharedRef<UGSCore::FWorkspaceUpdateContext, ESPMode::ThreadSafe> WorkspaceContext, UGSCore::EWorkspaceUpdateResult SyncResult, const FString& StatusMessage) {
 		OnWorkspaceSyncComplete(WorkspaceContext, SyncResult, StatusMessage);
 	};
 
-	CombinedSyncFilter = FUserSettings::GetCombinedSyncFilter(
+	CombinedSyncFilter = UGSCore::FUserSettings::GetCombinedSyncFilter(
 		Workspace->GetSyncCategories(),
 		UserSettings->SyncView,
 		UserSettings->SyncExcludedCategories,
@@ -530,26 +536,26 @@ void UGSTab::SetupWorkspace()
 		WorkspaceSettings->SyncExcludedCategories);
 
 	// Options on what to do with workspace when updating it
-	Options = EWorkspaceUpdateOptions::Sync | EWorkspaceUpdateOptions::SyncArchives | EWorkspaceUpdateOptions::GenerateProjectFiles;
+	Options = UGSCore::EWorkspaceUpdateOptions::Sync | UGSCore::EWorkspaceUpdateOptions::SyncArchives | UGSCore::EWorkspaceUpdateOptions::GenerateProjectFiles;
 	if (UserSettings->bAutoResolveConflicts)
 	{
-		Options |= EWorkspaceUpdateOptions::AutoResolveChanges;
+		Options |= UGSCore::EWorkspaceUpdateOptions::AutoResolveChanges;
 	}
 	if (UserSettings->bUseIncrementalBuilds)
 	{
-		Options |= EWorkspaceUpdateOptions::UseIncrementalBuilds;
+		Options |= UGSCore::EWorkspaceUpdateOptions::UseIncrementalBuilds;
 	}
 	if (UserSettings->bBuildAfterSync)
 	{
-		Options |= EWorkspaceUpdateOptions::Build;
+		Options |= UGSCore::EWorkspaceUpdateOptions::Build;
 	}
 	if (UserSettings->bBuildAfterSync && UserSettings->bRunAfterSync)
 	{
-		Options |= EWorkspaceUpdateOptions::RunAfterSync;
+		Options |= UGSCore::EWorkspaceUpdateOptions::RunAfterSync;
 	}
 	if (UserSettings->bOpenSolutionAfterSync)
 	{
-		Options |= EWorkspaceUpdateOptions::OpenSolutionAfterSync;
+		Options |= UGSCore::EWorkspaceUpdateOptions::OpenSolutionAfterSync;
 	}
 
 	// Setup our Perforce and Event monitoring threads
@@ -557,7 +563,7 @@ void UGSTab::SetupWorkspace()
 	FString SelectedClientFileName = DetectSettings->NewSelectedClientFileName;
 
 	// TODO create callback functions that will be queued for the main thread to generate and update the main table view
-	PerforceMonitor = MakeShared<FPerforceMonitor>(PerforceClient.ToSharedRef(), BranchClientPath, SelectedClientFileName, SelectedProjectIdentifier, ProjectLogBaseName + ".p4.log");
+	PerforceMonitor = MakeShared<UGSCore::FPerforceMonitor>(PerforceClient.ToSharedRef(), BranchClientPath, SelectedClientFileName, SelectedProjectIdentifier, ProjectLogBaseName + ".p4.log");
 	PerforceMonitor->OnUpdate = [this]{ QueueMessageForMainThread([this] { UpdateGameTabBuildList(); }); };
 
 	//PerforceMonitor->OnUpdateMetadata = [this]{ printf("PerforceMonitor->OnUpdateMetadata\n"); }; //MessageQueue.Add([this]{ UpdateBuildMetadata(); }); };
@@ -566,7 +572,7 @@ void UGSTab::SetupWorkspace()
 	/* TODO figure out if this is even working, and if so how to correctly use this
 	 */
 	FString SqlConnectionString;
-	EventMonitor = MakeShared<FEventMonitor>(SqlConnectionString, FPerforceUtils::GetClientOrDepotDirectoryName(*SelectedProjectIdentifier), PerforceClient->UserName, ProjectLogBaseName + ".review.log");
+	EventMonitor = MakeShared<UGSCore::FEventMonitor>(SqlConnectionString, UGSCore::FPerforceUtils::GetClientOrDepotDirectoryName(*SelectedProjectIdentifier), PerforceClient->UserName, ProjectLogBaseName + ".review.log");
 	EventMonitor->OnUpdatesReady = [this]{ printf("EventMonitor->OnUpdatesReady\n"); }; //MessageQueue.Add([this]{ UpdateReviews(); }); };
 
 	// Start the threads
diff --git a/Engine/Source/Programs/SlateUGS/Private/UGSTab.h b/Engine/Source/Programs/SlateUGS/Private/UGSTab.h
--- a/Engine/Source/Programs/SlateUGS/Private/UGSTab.h
+++ b/Engine/Source/Programs/SlateUGS/Private/UGSTab.h
@@ -8,12 +8,16 @@
 
 #include "UGSCore/Workspace.h"
 
-struct FUserWorkspaceSettings;
-struct FUserProjectSettings;
-struct FUserSettings;
-class FDetectProjectSettingsTask;
-class FPerforceMonitor;
-class FEventMonitor;
+namespace UGSCore
+{
+	struct FUserWorkspaceSettings;
+	struct FUserProjectSettings;
+	struct FUserSettings;
+	class FDetectProjectSettingsTask;
+	class FPerforceMonitor;
+	class FEventMonitor;
+}
+
 class UGSTabManager;
 
 enum SyncCategoryType
@@ -51,14 +55,14 @@ public:
 	FString GetSyncProgress() const;
 	const TArray<FString>& GetSyncFilters() const;
 	const TArray<FString>& GetCombinedSyncFilter() const;
-	TArray<FWorkspaceSyncCategory> GetSyncCategories(SyncCategoryType CategoryType) const;
+	TArray<UGSCore::FWorkspaceSyncCategory> GetSyncCategories(SyncCategoryType CategoryType) const;
 	TArray<FString> GetSyncViews(SyncCategoryType CategoryType) const;
 	UGSTabManager* GetTabManager();
 private:
 
 	void OnWorkspaceSyncComplete(
-		TSharedRef<FWorkspaceUpdateContext, ESPMode::ThreadSafe> WorkspaceContext,
-		EWorkspaceUpdateResult SyncResult,
+		TSharedRef<UGSCore::FWorkspaceUpdateContext, ESPMode::ThreadSafe> WorkspaceContext,
+		UGSCore::EWorkspaceUpdateResult SyncResult,
 		const FString& StatusMessage);
 
 	// Allows the queuing of functions from threads to be run on the main thread
@@ -81,18 +85,18 @@ private:
 	// Core data
 	UGSTabManager* TabManager = nullptr;
 	FString ProjectFileName;
-	TSharedPtr<FWorkspace> Workspace;
-	TSharedPtr<FPerforceConnection> PerforceClient;
-	TSharedPtr<FUserWorkspaceSettings> WorkspaceSettings;
-	TSharedPtr<FUserProjectSettings> ProjectSettings;
-	EWorkspaceUpdateOptions Options;
-	TSharedPtr<FDetectProjectSettingsTask> DetectSettings;
+	TSharedPtr<UGSCore::FWorkspace> Workspace;
+	TSharedPtr<UGSCore::FPerforceConnection> PerforceClient;
+	TSharedPtr<UGSCore::FUserWorkspaceSettings> WorkspaceSettings;
+	TSharedPtr<UGSCore::FUserProjectSettings> ProjectSettings;
+	UGSCore::EWorkspaceUpdateOptions Options;
+	TSharedPtr<UGSCore::FDetectProjectSettingsTask> DetectSettings;
 	TArray<FString> CombinedSyncFilter;
-	TSharedPtr<FUserSettings> UserSettings;
+	TSharedPtr<UGSCore::FUserSettings> UserSettings;
 
 	// Monitoring threads
-	TSharedPtr<FPerforceMonitor> PerforceMonitor;
-	TSharedPtr<FEventMonitor> EventMonitor;
+	TSharedPtr<UGSCore::FPerforceMonitor> PerforceMonitor;
+	TSharedPtr<UGSCore::FEventMonitor> EventMonitor;
 
 	// Queue for handling callbacks from multiple threads that need to be called on the main thread
 	TArray<TFunction<void()>> MessageQueue;
diff --git a/Engine/Source/Programs/SlateUGS/Private/Widgets/SGameSyncTab.cpp b/Engine/Source/Programs/SlateUGS/Private/Widgets/SGameSyncTab.cpp
--- a/Engine/Source/Programs/SlateUGS/Private/Widgets/SGameSyncTab.cpp
+++ b/Engine/Source/Programs/SlateUGS/Private/Widgets/SGameSyncTab.cpp
@@ -47,16 +47,16 @@ TSharedRef<SWidget> SBuildDataRow::GenerateWidgetForColumn(const FName& ColumnId
 
 		switch (CurrentItem->ReviewStatus)
 		{
-			case EReviewVerdict::Good:
+			case UGSCore::EReviewVerdict::Good:
 				StatusCircle->SetColorAndOpacity(FLinearColor::Green);
 				break;
-			case EReviewVerdict::Bad:
+			case UGSCore::EReviewVerdict::Bad:
 				StatusCircle->SetColorAndOpacity(FLinearColor::Red);
 				break;
-			case EReviewVerdict::Mixed:
+			case UGSCore::EReviewVerdict::Mixed:
 				StatusCircle->SetColorAndOpacity(FLinearColor::Yellow);
 				break;
-			case EReviewVerdict::Unknown:
+			case UGSCore::EReviewVerdict::Unknown:
 			default:
 				StatusCircle->SetColorAndOpacity(FLinearColor::Gray);
 				break;
diff --git a/Engine/Source/Programs/SlateUGS/Private/Widgets/SLogWidget.h b/Engine/Source/Programs/SlateUGS/Private/Widgets/SLogWidget.h
--- a/Engine/Source/Programs/SlateUGS/Private/Widgets/SLogWidget.h
+++ b/Engine/Source/Programs/SlateUGS/Private/Widgets/SLogWidget.h
@@ -3,10 +3,11 @@
 #pragma once
 
 #include "SlateCore.h"
-#include "UGSCore/OutputAdapters.h"
 #include "Framework/Text/BaseTextLayoutMarshaller.h"
 #include "Widgets/Input/SMultiLineEditableTextBox.h"
 
+#include "UGSCore/OutputAdapters.h"
+
 class FLogWidgetTextLayoutMarshaller : public FBaseTextLayoutMarshaller
 {
 public:
@@ -57,7 +58,7 @@ protected:
 	EActiveTimerReturnType OnTimerElapsed(double CurrentTime, float DeltaTime);
 };
 
-class FLogWidgetTextWriter : public FLineBasedTextWriter
+class FLogWidgetTextWriter : public UGSCore::FLineBasedTextWriter
 {
 public:
 	FLogWidgetTextWriter(const TSharedRef<SLogWidget>& InLogWidget)
diff --git a/Engine/Source/Programs/SlateUGS/Private/Widgets/SModalTaskWindow.cpp b/Engine/Source/Programs/SlateUGS/Private/Widgets/SModalTaskWindow.cpp
--- a/Engine/Source/Programs/SlateUGS/Private/Widgets/SModalTaskWindow.cpp
+++ b/Engine/Source/Programs/SlateUGS/Private/Widgets/SModalTaskWindow.cpp
@@ -5,6 +5,7 @@
 #include "Widgets/SBoxPanel.h"
 #include "Widgets/Text/STextBlock.h"
 #include "Widgets/Images/SThrobber.h"
+#include "UGSCore/ModalTask.h"
 
 SModalTaskWindow::SModalTaskWindow()
 {
@@ -80,7 +81,7 @@ uint32 SModalTaskWindow::Run()
 	return 0;
 }
 
-TSharedRef<FModalTaskResult> ExecuteModalTask(TSharedPtr<SWidget> Parent, TSharedRef<IModalTask> Task, const FText& InTitle, const FText& InMessage)
+TSharedRef<UGSCore::FModalTaskResult> ExecuteModalTask(TSharedPtr<SWidget> Parent, TSharedRef<UGSCore::IModalTask> Task, const FText& InTitle, const FText& InMessage)
 {
 	TSharedRef<SModalTaskWindow> Window =
 		SNew(SModalTaskWindow)
diff --git a/Engine/Source/Programs/SlateUGS/Private/Widgets/SModalTaskWindow.h b/Engine/Source/Programs/SlateUGS/Private/Widgets/SModalTaskWindow.h
--- a/Engine/Source/Programs/SlateUGS/Private/Widgets/SModalTaskWindow.h
+++ b/Engine/Source/Programs/SlateUGS/Private/Widgets/SModalTaskWindow.h
@@ -6,7 +6,12 @@
 #include "SlateCore.h"
 #include "HAL/Runnable.h"
 #include "HAL/RunnableThread.h"
-#include "UGSCore/ModalTask.h"
+
+namespace UGSCore
+{
+	class IModalTask;
+	class FModalTaskResult;
+}
 
 class SModalTaskWindow : public SWindow, private FRunnable
 {
@@ -14,10 +19,10 @@ public:
 	SLATE_BEGIN_ARGS(SModalTaskWindow) {}
 		SLATE_ARGUMENT(FText, Title)
 		SLATE_ARGUMENT(FText, Message)
-		SLATE_ARGUMENT(TSharedPtr<IModalTask>, Task)
+		SLATE_ARGUMENT(TSharedPtr<UGSCore::IModalTask>, Task)
 	SLATE_END_ARGS()
 
-	TSharedPtr<FModalTaskResult> Result;
+	TSharedPtr<UGSCore::FModalTaskResult> Result;
 
 	SModalTaskWindow();
 	~SModalTaskWindow();
@@ -29,9 +34,9 @@ private:
 	FEvent* AbortEvent;
 	FEvent* CloseEvent;
 	FRunnableThread* Thread;
-	TSharedPtr<IModalTask> Task;
+	TSharedPtr<UGSCore::IModalTask> Task;
 
 	virtual uint32 Run() override;
 };
 
-TSharedRef<FModalTaskResult> ExecuteModalTask(TSharedPtr<SWidget> Parent, TSharedRef<IModalTask> Task, const FText& InTitle, const FText& InMessage);
+TSharedRef<UGSCore::FModalTaskResult> ExecuteModalTask(TSharedPtr<SWidget> Parent, TSharedRef<UGSCore::IModalTask> Task, const FText& InTitle, const FText& InMessage);
diff --git a/Engine/Source/Programs/SlateUGS/Private/Widgets/SSyncFilterWindow.cpp b/Engine/Source/Programs/SlateUGS/Private/Widgets/SSyncFilterWindow.cpp
--- a/Engine/Source/Programs/SlateUGS/Private/Widgets/SSyncFilterWindow.cpp
+++ b/Engine/Source/Programs/SlateUGS/Private/Widgets/SSyncFilterWindow.cpp
@@ -15,10 +15,12 @@
 void SSyncFilterWindow::ConstructSyncFilters()
 {
 	WorkspaceCategoriesCurrent = Tab->GetSyncCategories(SyncCategoryType::CurrentWorkspace);
+
 	SyncFiltersCurrent = SNew(SCheckBoxList)
 		.ItemHeaderLabel(LOCTEXT("CheckBoxListCurrent", "Sync Filters for the current workspace"))
 		.IncludeGlobalCheckBoxInHeaderRow(false);
-	for (const FWorkspaceSyncCategory& Category : WorkspaceCategoriesCurrent)
+
+	for (const UGSCore::FWorkspaceSyncCategory& Category : WorkspaceCategoriesCurrent)
 	{
 		TSharedRef<SBox> CheckBoxWrapper = SNew(SBox)
 			.Padding(5.0f)
@@ -26,14 +28,17 @@ void SSyncFilterWindow::ConstructSyncFilters()
 				SNew(STextBlock)
 				.Text(FText::FromString(Category.Name))
 			];
+
 		SyncFiltersCurrent->AddItem(CheckBoxWrapper, Category.bEnable);
 	}
 
 	WorkspaceCategoriesAll = Tab->GetSyncCategories(SyncCategoryType::AllWorkspaces);
+
 	SyncFiltersAll = SNew(SCheckBoxList)
 		.ItemHeaderLabel(LOCTEXT("CheckBoxListAll", "Sync Filters for all workspaces"))
 		.IncludeGlobalCheckBoxInHeaderRow(false);
-	for (const FWorkspaceSyncCategory& Category : WorkspaceCategoriesAll)
+
+	for (const UGSCore::FWorkspaceSyncCategory& Category : WorkspaceCategoriesAll)
 	{
 		TSharedRef<SBox> CheckBoxWrapper = SNew(SBox)
 			.Padding(5.0f)
@@ -41,6 +46,7 @@ void SSyncFilterWindow::ConstructSyncFilters()
 				SNew(STextBlock)
 				.Text(FText::FromString(Category.Name))
 			];
+
 		SyncFiltersAll->AddItem(CheckBoxWrapper, Category.bEnable);
 	}
 }
diff --git a/Engine/Source/Programs/SlateUGS/Private/Widgets/SSyncFilterWindow.h b/Engine/Source/Programs/SlateUGS/Private/Widgets/SSyncFilterWindow.h
--- a/Engine/Source/Programs/SlateUGS/Private/Widgets/SSyncFilterWindow.h
+++ b/Engine/Source/Programs/SlateUGS/Private/Widgets/SSyncFilterWindow.h
@@ -30,8 +30,8 @@ private:
 	void ConstructCustomSyncViewTextBoxes();
 
 	UGSTab* Tab;
-	TArray<FWorkspaceSyncCategory> WorkspaceCategoriesCurrent;
-	TArray<FWorkspaceSyncCategory> WorkspaceCategoriesAll;
+	TArray<UGSCore::FWorkspaceSyncCategory> WorkspaceCategoriesCurrent;
+	TArray<UGSCore::FWorkspaceSyncCategory> WorkspaceCategoriesAll;
 
 	TSharedPtr<SCheckBoxList> SyncFiltersCurrent;
 	TSharedPtr<SCheckBoxList> SyncFiltersAll;
