Fix some hard coded GPF.sh -> GPF.bat
Fix ExecProc not spawning with hidden for cmd

diff --git a/Engine/Source/Programs/SlateUGS/Private/UGSCore/Utility.cpp b/Engine/Source/Programs/SlateUGS/Private/UGSCore/Utility.cpp
--- a/Engine/Source/Programs/SlateUGS/Private/UGSCore/Utility.cpp
+++ b/Engine/Source/Programs/SlateUGS/Private/UGSCore/Utility.cpp
@@ -116,7 +116,7 @@ int FUtility::ExecuteProcess(const TCHAR* FileName, const TCHAR* CommandLine, co
     void* WritePipe = nullptr;
     FPlatformProcess::CreatePipe(ReadPipe, WritePipe);
 
-	FProcHandle Proc = FPlatformProcess::CreateProc(FileName, CommandLine, false, false, false, &ProcId, 0, nullptr, WritePipe, ReadPipe);
+	FProcHandle Proc = FPlatformProcess::CreateProc(FileName, CommandLine, false, true, true, &ProcId, 0, nullptr, WritePipe, ReadPipe);
 
 	FString Output;
     FString LatestOutput = FPlatformProcess::ReadPipe(ReadPipe);
@@ -150,7 +150,7 @@ int FUtility::ExecuteProcess(const TCHAR* FileName, const TCHAR* CommandLine, co
     void* WritePipe = nullptr;
     FPlatformProcess::CreatePipe(ReadPipe, WritePipe);
 
-	FProcHandle Proc = FPlatformProcess::CreateProc(FileName, CommandLine, false, false, false, &ProcId, 0, nullptr, WritePipe, ReadPipe);
+	FProcHandle Proc = FPlatformProcess::CreateProc(FileName, CommandLine, false, true, true, &ProcId, 0, nullptr, WritePipe, ReadPipe);
 
     FString LatestOutput = FPlatformProcess::ReadPipe(ReadPipe);
     while (FPlatformProcess::IsProcRunning(Proc) || !LatestOutput.IsEmpty())
diff --git a/Engine/Source/Programs/SlateUGS/Private/UGSCore/Workspace.cpp b/Engine/Source/Programs/SlateUGS/Private/UGSCore/Workspace.cpp
--- a/Engine/Source/Programs/SlateUGS/Private/UGSCore/Workspace.cpp
+++ b/Engine/Source/Programs/SlateUGS/Private/UGSCore/Workspace.cpp
@@ -517,7 +517,11 @@ EWorkspaceUpdateResult FWorkspace::UpdateWorkspaceInternal(FWorkspaceUpdateConte
 				if(!Perforce->GetActiveStream(BranchOrStreamName, AbortEvent, Log.Get()))
 				{
 					FString DepotFileName;
+#if PLATFORM_WINDOWS
+					if(!Perforce->ConvertToDepotPath(ClientRootPath + TEXT("/GenerateProjectFiles.bat"), DepotFileName, AbortEvent, Log.Get()))
+#else
 					if(!Perforce->ConvertToDepotPath(ClientRootPath + TEXT("/GenerateProjectFiles.sh"), DepotFileName, AbortEvent, Log.Get()))
+#endif
 					{
 						OutStatusMessage = FString::Printf(TEXT("Couldn't determine branch name for %s."), *SelectedClientFileName);
 						return EWorkspaceUpdateResult::FailedToSync;
@@ -784,7 +788,7 @@ EWorkspaceUpdateResult FWorkspace::UpdateWorkspaceInternal(FWorkspaceUpdateConte
 		}
 
 #if PLATFORM_WINDOWS
-		FString CommandLine = FString::Printf(TEXT("/C \"\"%s\" %s-progress\""), *(LocalRootPath / TEXT("GenerateProjectFiles.sh")), *ProjectFileArgument);
+		FString CommandLine = FString::Printf(TEXT("/C \"\"%s\" %s-progress\""), *(LocalRootPath / TEXT("GenerateProjectFiles.bat")), *ProjectFileArgument);
 #else
 		FString CommandLine = FString::Printf(TEXT("\"%s\" %s-progress"), *(LocalRootPath / TEXT("GenerateProjectFiles.sh")), *ProjectFileArgument);
 #endif
