Add support for opening p4v/p4vc if its in the PATH
Add 3 new menus when right clicking a build list item
  View in Swarm (opens the CL in swarm)
  Copy Changelist (copies the CL to the clipboard)
  More Info (opens the CL in p4v)
Fixed a bug in Perforce.cpp not having a ServerAndPort setup

diff --git a/Engine/Source/Programs/SlateUGS/Private/UGSCore/DetectProjectSettingsTask.cpp b/Engine/Source/Programs/SlateUGS/Private/UGSCore/DetectProjectSettingsTask.cpp
--- a/Engine/Source/Programs/SlateUGS/Private/UGSCore/DetectProjectSettingsTask.cpp
+++ b/Engine/Source/Programs/SlateUGS/Private/UGSCore/DetectProjectSettingsTask.cpp
@@ -42,11 +42,8 @@ TSharedRef<FModalTaskResult> FDetectProjectSettingsTask::RunInternal(FEvent* Abo
 
 //	Directory.SetCurrentDirectory(Path.GetDirectoryName(NewSelectedFileName));
 //
-//	string ServerAndPort;
-//	if (Perforce.GetSetting("P4PORT", out ServerAndPort, Log))
-//	{
-//		Perforce = new PerforceConnection(null, null, ServerAndPort);
-//	}
+	FString ServerAndPort;
+	Perforce->GetSetting(TEXT("P4PORT"), ServerAndPort, AbortEvent, Log.Get());
 
 	// Get the Perforce server info
 	TSharedPtr<FPerforceInfoRecord> PerforceInfo;
@@ -64,6 +61,11 @@ TSharedRef<FModalTaskResult> FDetectProjectSettingsTask::RunInternal(FEvent* Abo
 	}
 	ServerTimeZone = PerforceInfo->ServerTimeZone;
 
+	if (ServerAndPort.IsEmpty())
+	{
+		ServerAndPort = PerforceInfo->ServerAddress;
+	}
+
 	// Find all the clients on this machine
 	Log->Logf(TEXT("Enumerating clients on local machine..."));
 	TArray<FPerforceClientRecord> Clients;
@@ -88,7 +90,7 @@ TSharedRef<FModalTaskResult> FDetectProjectSettingsTask::RunInternal(FEvent* Abo
 					continue;
 				}
 
-				TSharedRef<FPerforceConnection> CandidateClient = MakeShared<FPerforceConnection>(*PerforceInfo->UserName, *Client.Name, *Perforce->ServerAndPort);
+				TSharedRef<FPerforceConnection> CandidateClient = MakeShared<FPerforceConnection>(*PerforceInfo->UserName, *Client.Name, *ServerAndPort);
 
 				bool bFileExists;
 				if(!CandidateClient->FileExists(NewSelectedFileName, bFileExists, AbortEvent, Log.Get()) || !bFileExists)
diff --git a/Engine/Source/Programs/SlateUGS/Private/UGSCore/Perforce.cpp b/Engine/Source/Programs/SlateUGS/Private/UGSCore/Perforce.cpp
--- a/Engine/Source/Programs/SlateUGS/Private/UGSCore/Perforce.cpp
+++ b/Engine/Source/Programs/SlateUGS/Private/UGSCore/Perforce.cpp
@@ -1446,6 +1446,46 @@ bool FPerforceConnection::RunCommandWithBinaryOutput(const FString& CommandLine,
 	return FPlatformProcess::GetProcReturnCode(P4Proc, &ExitCode) == 0;
 }
 
+void FPerforceConnection::OpenP4V(const FString& AdditionalArgs) const
+{
+#if PLATFORM_WINDOWS
+	const TCHAR* P4VExe = TEXT("p4v.exe");
+#else
+	const TCHAR* P4VExe = TEXT("p4v");
+#endif
+
+	FString P4VArgs = GetArgumentsForExternalProgram() + TEXT(" ") + AdditionalArgs;
+
+	FPlatformProcess::CreateProc(P4VExe, *P4VArgs, true, true, true, nullptr, 0, nullptr, nullptr);
+	printf("%s\n", TCHAR_TO_ANSI(*P4VArgs));
+}
+
+void FPerforceConnection::OpenP4VC(const FString& AdditionalArgs) const
+{
+#if PLATFORM_WINDOWS
+	const TCHAR* P4VCExe = TEXT("p4vc.exe");
+#else
+	const TCHAR* P4VCExe = TEXT("p4vc");
+#endif
+
+	FString P4VCArgs = GetArgumentsForExternalProgram() + TEXT(" ") + AdditionalArgs;
+
+	FPlatformProcess::CreateProc(P4VCExe, *P4VCArgs, true, true, true, nullptr, 0, nullptr, nullptr);
+	printf("%s\n", TCHAR_TO_ANSI(*P4VCArgs));
+}
+
+FString FPerforceConnection::GetArgumentsForExternalProgram(bool bIncludeClient) const
+{
+	FString ExternalArgs = FString::Printf(TEXT("-p \"%s\" -u \"%s\""), *ServerAndPort, *UserName);
+
+	if (bIncludeClient && !ClientName.IsEmpty())
+	{
+		ExternalArgs += FString::Printf(TEXT(" -c \"%s\""), *ClientName);
+	}
+
+	return ExternalArgs;
+}
+
 //// FPerforceUtils ////
 
 FString FPerforceUtils::GetClientOrDepotDirectoryName(const TCHAR* ClientFile)
diff --git a/Engine/Source/Programs/SlateUGS/Private/UGSCore/Perforce.h b/Engine/Source/Programs/SlateUGS/Private/UGSCore/Perforce.h
--- a/Engine/Source/Programs/SlateUGS/Private/UGSCore/Perforce.h
+++ b/Engine/Source/Programs/SlateUGS/Private/UGSCore/Perforce.h
@@ -232,6 +232,9 @@ public:
 	bool Sync(const TArray<FString>& DepotPaths, int ChangeNumber, TFunction<void(const FPerforceFileRecord&)> SyncOutput, TArray<FString>& OutTamperedFiles, const FPerforceSyncOptions* Options, FEvent* AbortEvent, FOutputDevice& Log) const;
 	bool LatestChangeList(int& OutChangeList, FEvent* AbortEvent, FOutputDevice& Log) const;
 
+	void OpenP4V(const FString& AdditionalArgs = FString()) const;
+	void OpenP4VC(const FString& AdditionalArgs = FString()) const;
+
 private:
 	static bool FilterSyncOutput(const FPerforceOutputLine& Line, FPerforceTagRecordParser& Parser, TArray<FString>& OutTamperedFiles, FOutputDevice& Log);
 	static void ParseTamperedFile(const FString& Line, TArray<FString>& OutTamperedFiles);
@@ -262,6 +265,8 @@ private:
 
 	bool RunCommandWithBinaryOutput(const FString& CommandLine, TArray<TMap<FString, FString>>& OutRecords, ECommandOptions Options, FEvent* AbortEvent, FOutputDevice& Log) const;
 	bool RunCommandWithBinaryOutput(const FString& CommandLine, TFunction<void(const TMap<FString, FString>&)> HandleOutput, ECommandOptions Options, FEvent* AbortEvent, FOutputDevice& Log) const;
+
+	FString GetArgumentsForExternalProgram(bool bIncludeClient = true) const;
 };
 
 ENUM_CLASS_FLAGS(FPerforceConnection::ECommandOptions)
diff --git a/Engine/Source/Programs/SlateUGS/Private/UGSTab.cpp b/Engine/Source/Programs/SlateUGS/Private/UGSTab.cpp
--- a/Engine/Source/Programs/SlateUGS/Private/UGSTab.cpp
+++ b/Engine/Source/Programs/SlateUGS/Private/UGSTab.cpp
@@ -16,6 +16,8 @@
 #include "Widgets/SModalTaskWindow.h"
 #include "Widgets/SLogWidget.h"
 
+#include "HAL/PlatformApplicationMisc.h"
+
 #include "Internationalization/Regex.h"
 
 #define LOCTEXT_NAMESPACE "UGSTab"
@@ -208,6 +210,28 @@ void UGSTab::OnSyncChangelist(int Changelist)
 	Workspace->Update(Context);
 }
 
+void UGSTab::OnViewInSwarmClicked(int Changelist) const
+{
+	FString SwarmURL = FString::Printf(TEXT("https://p4-swarm.epicgames.net/changes/%i"), Changelist);
+
+	FPlatformProcess::LaunchURL(*SwarmURL, nullptr, nullptr);
+}
+
+void UGSTab::OnCopyChangeListClicked(int Changelist) const
+{
+	FPlatformApplicationMisc::ClipboardCopy(*FString::FromInt(Changelist));
+}
+
+void UGSTab::OnMoreInfoClicked(int Changelist) const
+{
+	PerforceClient->OpenP4VC(FString::Printf(TEXT("change %i"), Changelist));
+}
+
+void UGSTab::OnOpenPerforceClicked() const
+{
+	PerforceClient->OpenP4V();
+}
+
 void UGSTab::OnSyncLatest()
 {
 	int ChangeNumber = -1;
diff --git a/Engine/Source/Programs/SlateUGS/Private/UGSTab.h b/Engine/Source/Programs/SlateUGS/Private/UGSTab.h
--- a/Engine/Source/Programs/SlateUGS/Private/UGSTab.h
+++ b/Engine/Source/Programs/SlateUGS/Private/UGSTab.h
@@ -43,6 +43,10 @@ public:
 	// Slate callbacks
 	bool OnWorkspaceChosen(const FString& Project);
 	void OnSyncChangelist(int Changelist);
+	void OnViewInSwarmClicked(int Changelist) const;
+	void OnCopyChangeListClicked(int Changelist) const;
+	void OnMoreInfoClicked(int Changelist) const;
+	void OnOpenPerforceClicked() const;
 	void OnSyncLatest();
 	void OnSyncFilterWindowSaved(
 		const TArray<FString>& SyncViewCurrent,
diff --git a/Engine/Source/Programs/SlateUGS/Private/Widgets/SGameSyncTab.cpp b/Engine/Source/Programs/SlateUGS/Private/Widgets/SGameSyncTab.cpp
--- a/Engine/Source/Programs/SlateUGS/Private/Widgets/SGameSyncTab.cpp
+++ b/Engine/Source/Programs/SlateUGS/Private/Widgets/SGameSyncTab.cpp
@@ -235,7 +235,7 @@ void SGameSyncTab::Construct(const FArguments& InArgs)
 						SNew(SSimpleButton)
 						.Text(LOCTEXT("OpenPerforce", "Open Perforce"))
 						.Icon(FAppStyle::Get().GetBrush("Icons.Blueprints")) // Todo: shouldn't use this icon (repurposing)
-						.IsEnabled(false) // Todo: enable after adding this functionality
+						.OnClicked_Lambda([this] { Tab->OnOpenPerforceClicked(); return FReply::Handled(); })
 					]
 					+SHorizontalBox::Slot()
 					.AutoWidth()
@@ -484,6 +484,35 @@ TSharedPtr<SWidget> SGameSyncTab::OnRightClickedBuild()
 			NAME_None,
 			EUserInterfaceActionType::Button
 		);
+
+		MenuBuilder.AddSeparator();
+
+		MenuBuilder.AddMenuEntry(
+			FText::FromString(TEXT("View in Swarm...")),
+			TAttribute<FText>(),
+			FSlateIcon(),
+			FUIAction(FExecuteAction::CreateRaw(Tab, &UGSTab::OnViewInSwarmClicked, SelectedItems[0]->Changelist)),
+			NAME_None,
+			EUserInterfaceActionType::Button
+		);
+
+		MenuBuilder.AddMenuEntry(
+			FText::FromString(TEXT("Copy Changelist")),
+			TAttribute<FText>(),
+			FSlateIcon(),
+			FUIAction(FExecuteAction::CreateRaw(Tab, &UGSTab::OnCopyChangeListClicked, SelectedItems[0]->Changelist)),
+			NAME_None,
+			EUserInterfaceActionType::Button
+		);
+
+		MenuBuilder.AddMenuEntry(
+			FText::FromString(TEXT("More Info")),
+			TAttribute<FText>(),
+			FSlateIcon(),
+			FUIAction(FExecuteAction::CreateRaw(Tab, &UGSTab::OnMoreInfoClicked, SelectedItems[0]->Changelist)),
+			NAME_None,
+			EUserInterfaceActionType::Button
+		);
 	}
 
 	return MenuBuilder.MakeWidget();
