Fix exit code checking, check the exit code not the return value of the function

diff --git a/Engine/Source/Programs/SlateUGS/Private/UGSCore/Perforce.cpp b/Engine/Source/Programs/SlateUGS/Private/UGSCore/Perforce.cpp
--- a/Engine/Source/Programs/SlateUGS/Private/UGSCore/Perforce.cpp
+++ b/Engine/Source/Programs/SlateUGS/Private/UGSCore/Perforce.cpp
@@ -1443,7 +1443,9 @@ bool FPerforceConnection::RunCommandWithBinaryOutput(const FString& CommandLine,
 	int ExitCode = -1;
 
 	checkf(BufferPos == BufferEnd, TEXT("%d bytes of incomplete record data received from P4"), BufferEnd - BufferPos);
-	return FPlatformProcess::GetProcReturnCode(P4Proc, &ExitCode) == 0;
+	FPlatformProcess::GetProcReturnCode(P4Proc, &ExitCode);
+
+	return ExitCode == 0;
 }
 
 void FPerforceConnection::OpenP4V(const FString& AdditionalArgs) const
@@ -1457,7 +1459,6 @@ void FPerforceConnection::OpenP4V(const FString& AdditionalArgs) const
 	FString P4VArgs = GetArgumentsForExternalProgram() + TEXT(" ") + AdditionalArgs;
 
 	FPlatformProcess::CreateProc(P4VExe, *P4VArgs, true, true, true, nullptr, 0, nullptr, nullptr);
-	printf("%s\n", TCHAR_TO_ANSI(*P4VArgs));
 }
 
 void FPerforceConnection::OpenP4VC(const FString& AdditionalArgs) const
@@ -1471,7 +1472,6 @@ void FPerforceConnection::OpenP4VC(const FString& AdditionalArgs) const
 	FString P4VCArgs = GetArgumentsForExternalProgram() + TEXT(" ") + AdditionalArgs;
 
 	FPlatformProcess::CreateProc(P4VCExe, *P4VCArgs, true, true, true, nullptr, 0, nullptr, nullptr);
-	printf("%s\n", TCHAR_TO_ANSI(*P4VCArgs));
 }
 
 FString FPerforceConnection::GetArgumentsForExternalProgram(bool bIncludeClient) const
