Update Perforce monitor to check if it contains Code and/or Content
Need to reflect this in the UI now

diff --git a/Engine/Source/Programs/SlateUGS/Private/UGSCore/PerforceMonitor.cpp b/Engine/Source/Programs/SlateUGS/Private/UGSCore/PerforceMonitor.cpp
--- a/Engine/Source/Programs/SlateUGS/Private/UGSCore/PerforceMonitor.cpp
+++ b/Engine/Source/Programs/SlateUGS/Private/UGSCore/PerforceMonitor.cpp
@@ -97,16 +97,17 @@ TArray<TSharedRef<FPerforceChangeSummary, ESPMode::ThreadSafe>> FPerforceMonitor
 	return Changes;
 }
 
-bool FPerforceMonitor::TryGetChangeType(int ChangeNumber, EPerforceChangeType& OutType) const
+bool FPerforceMonitor::TryGetChangeType(int ChangeNumber, FChangeType& OutType) const
 {
 	FScopeLock Lock(&CriticalSection);
 
-	const EPerforceChangeType* ChangeType = ChangeNumberToType.Find(ChangeNumber);
+	const FChangeType* ChangeType = ChangeNumberToType.Find(ChangeNumber);
 	if(ChangeType != nullptr)
 	{
 		OutType = *ChangeType;
 		return true;
 	}
+
 	return false;
 }
 
@@ -356,24 +357,46 @@ bool FPerforceMonitor::UpdateChangeTypes()
 		}
 	}
 
+	printf("UPDATE CHANGE TYPE ....................\n");
+
 	// Update them in batches
 	for(int QueryChangeNumber : QueryChangeNumbers)
 	{
+		printf("UPDATE CHANGE TYPE .................... %i\n", QueryChangeNumber);
+
 		// If there's something to check for, find all the content changes after this changelist
 		TSharedPtr<FPerforceDescribeRecord> DescribeRecord;
-		if(Perforce->Describe(QueryChangeNumber, DescribeRecord, AbortEvent, LogWriter))
+		if (Perforce->Describe(QueryChangeNumber, DescribeRecord, AbortEvent, LogWriter))
 		{
-			// Check whether the files are code or content
-			EPerforceChangeType Type = EPerforceChangeType::Content;
-			for(const FPerforceDescribeFileRecord& File : DescribeRecord->Files)
+			FChangeType Type;
+
+			for (const FPerforceDescribeFileRecord& File : DescribeRecord->Files)
 			{
 				const FString& DepotFile = File.DepotFile;
-				if(DepotFile.EndsWith(TEXT(".cs")) || DepotFile.EndsWith(TEXT(".h")) || DepotFile.EndsWith(TEXT(".cpp")) || DepotFile.EndsWith(TEXT(".usf")))
+
+				// TODO move this to a map, or something, way better way to do this then whats going on here
+				if (DepotFile.EndsWith(TEXT(".c")) || DepotFile.EndsWith(TEXT(".cc")) || DepotFile.EndsWith(TEXT(".cpp")) ||
+					DepotFile.EndsWith(TEXT(".inl")) || DepotFile.EndsWith(TEXT(".m")) || DepotFile.EndsWith(TEXT(".mm")) ||
+					DepotFile.EndsWith(TEXT(".rc")) || DepotFile.EndsWith(TEXT(".cs")) || DepotFile.EndsWith(TEXT(".csproj")) ||
+					DepotFile.EndsWith(TEXT(".h")) || DepotFile.EndsWith(TEXT(".hpp")) || DepotFile.EndsWith(TEXT(".inl")) ||
+					DepotFile.EndsWith(TEXT(".usf")) || DepotFile.EndsWith(TEXT(".ush")) || DepotFile.EndsWith(TEXT(".uproject")) ||
+					DepotFile.EndsWith(TEXT(".uplugin")) || DepotFile.EndsWith(TEXT(".sln")))
+				{
+					Type.bContainsCode = true;
+				}
+				else
+				{
+					Type.bContainsContent = true;
+				}
+
+				if (Type.bContainsCode && Type.bContainsContent)
 				{
-					Type = EPerforceChangeType::Code;
 					break;
 				}
 			}
+
+			printf("UPDATE CHANGE TYPE .................... WE GET HERE?? %i %i\n", Type.bContainsCode, Type.bContainsContent);
+
 			// Update the type of this change
 			{
 				FScopeLock Lock(&CriticalSection);
@@ -390,8 +413,8 @@ bool FPerforceMonitor::UpdateChangeTypes()
 		{
 			if(Change->User == Perforce->UserName)
 			{
-				EPerforceChangeType* Type = ChangeNumberToType.Find(Change->Number);
-				if(Type != nullptr && *Type == EPerforceChangeType::Code)
+				FChangeType* Type = ChangeNumberToType.Find(Change->Number);
+				if(Type != nullptr && Type->bContainsCode)
 				{
 					NewLastCodeChangeByCurrentUser = FMath::Max(NewLastCodeChangeByCurrentUser, Change->Number);
 				}
diff --git a/Engine/Source/Programs/SlateUGS/Private/UGSCore/PerforceMonitor.h b/Engine/Source/Programs/SlateUGS/Private/UGSCore/PerforceMonitor.h
--- a/Engine/Source/Programs/SlateUGS/Private/UGSCore/PerforceMonitor.h
+++ b/Engine/Source/Programs/SlateUGS/Private/UGSCore/PerforceMonitor.h
@@ -12,10 +12,10 @@
 namespace UGSCore
 {
 
-enum class EPerforceChangeType
+struct FChangeType
 {
-	Code,
-	Content,
+	bool bContainsCode = false;
+	bool bContainsContent = false;
 };
 
 class FPerforceMonitor : FRunnable
@@ -41,7 +41,7 @@ public:
 	bool HasZippedBinaries() const;
 
 	TArray<FChangeSharedRef> GetChanges() const;
-	bool TryGetChangeType(int ChangeNumber, EPerforceChangeType& OutType) const;
+	bool TryGetChangeType(int ChangeNumber, FChangeType& OutType) const;
 	bool TryGetArchivePathForChangeNumber(int ChangeNumber, FString& OutArchivePath) const;
 	TSet<int> GetPromotedChangeNumbers() const;
 
@@ -55,7 +55,7 @@ private:
 	FBoundedLogWriter LogWriter;
 	FRunnableThread* WorkerThread;
 	TArray<FChangeSharedRef> Changes; // Sorted
-	TMap<int, EPerforceChangeType> ChangeNumberToType;
+	TMap<int, FChangeType> ChangeNumberToType;
 	TSet<int> PromotedChangeNumbers;
 	int ZippedBinariesConfigChangeNumber;
 	FString ZippedBinariesPath;
