Remove debug print

diff --git a/Engine/Source/Programs/SlateUGS/Private/UGSCore/PerforceMonitor.cpp b/Engine/Source/Programs/SlateUGS/Private/UGSCore/PerforceMonitor.cpp
--- a/Engine/Source/Programs/SlateUGS/Private/UGSCore/PerforceMonitor.cpp
+++ b/Engine/Source/Programs/SlateUGS/Private/UGSCore/PerforceMonitor.cpp
@@ -357,13 +357,9 @@ bool FPerforceMonitor::UpdateChangeTypes()
 		}
 	}
 
-	printf("UPDATE CHANGE TYPE ....................\n");
-
 	// Update them in batches
 	for(int QueryChangeNumber : QueryChangeNumbers)
 	{
-		printf("UPDATE CHANGE TYPE .................... %i\n", QueryChangeNumber);
-
 		// If there's something to check for, find all the content changes after this changelist
 		TSharedPtr<FPerforceDescribeRecord> DescribeRecord;
 		if (Perforce->Describe(QueryChangeNumber, DescribeRecord, AbortEvent, LogWriter))
@@ -395,8 +391,6 @@ bool FPerforceMonitor::UpdateChangeTypes()
 				}
 			}
 
-			printf("UPDATE CHANGE TYPE .................... WE GET HERE?? %i %i\n", Type.bContainsCode, Type.bContainsContent);
-
 			// Update the type of this change
 			{
 				FScopeLock Lock(&CriticalSection);
