Added full descriptions in tooltips of rows in the build list.
Fixed some margins and spacing for when the window is resized.

diff --git a/Engine/Source/Programs/SlateUGS/Private/ChangeInfo.h b/Engine/Source/Programs/SlateUGS/Private/ChangeInfo.h
--- a/Engine/Source/Programs/SlateUGS/Private/ChangeInfo.h
+++ b/Engine/Source/Programs/SlateUGS/Private/ChangeInfo.h
@@ -12,7 +12,7 @@ struct FChangeInfo
 	UGSCore::EReviewVerdict ReviewStatus = UGSCore::EReviewVerdict::Unknown;
 	int Changelist = 0;
 	FText Author;
-	FText Description;
+	FString Description;
 
 	// Todo: add Horde badges
 };
diff --git a/Engine/Source/Programs/SlateUGS/Private/UGSTab.cpp b/Engine/Source/Programs/SlateUGS/Private/UGSTab.cpp
--- a/Engine/Source/Programs/SlateUGS/Private/UGSTab.cpp
+++ b/Engine/Source/Programs/SlateUGS/Private/UGSTab.cpp
@@ -391,10 +391,6 @@ void UGSTab::UpdateGameTabBuildList()
 					ChangeInfos.Add(HeaderRow);
 				}
 
-				// split on \n so we can just take the first line
-				TArray<FString> DescLines;
-				Change.Description.ParseIntoArray(DescLines, TEXT("\n"));
-
 				ChangeInfos.Add(MakeShareable(new FChangeInfo
 				{
 					DisplayTime,
@@ -402,7 +398,7 @@ void UGSTab::UpdateGameTabBuildList()
 					Status,
 					Change.Number,
 					FText::FromString(FormatUserName(Change.User)),
-					FText::FromString(DescLines[0].TrimStartAndEnd())
+					Change.Description
 				}));
 			}
 		}
diff --git a/Engine/Source/Programs/SlateUGS/Private/Widgets/SGameSyncTab.cpp b/Engine/Source/Programs/SlateUGS/Private/Widgets/SGameSyncTab.cpp
--- a/Engine/Source/Programs/SlateUGS/Private/Widgets/SGameSyncTab.cpp
+++ b/Engine/Source/Programs/SlateUGS/Private/Widgets/SGameSyncTab.cpp
@@ -85,12 +85,16 @@ TSharedRef<SWidget> SBuildDataRow::GenerateWidgetForColumn(const FName& ColumnId
 	}
 	if (ColumnId == HordeTableColumnDescription)
 	{
-		TextItem->SetText(CurrentItem->Description);
+		// Perforce descriptions mix line endings, so remove all possible line endings
+		FString Description = CurrentItem->Description
+			.Replace(TEXT("\r"), TEXT(" "))
+			.Replace(TEXT("\n"), TEXT(" ")); // Todo: replacing both with spaces causes lots of double spaces, maybe filter that out?
+		TextItem->SetText(FText::FromString(Description));
 	}
 
 	return SNew(SBox)
 		.VAlign(VAlign_Center)
-		.Padding(10.0f, 2.5f)
+		.Padding(5.0f, 2.5f)
 		[
 			TextItem
 		];
@@ -125,7 +129,7 @@ TSharedRef<ITableRow> SGameSyncTab::GenerateHordeBuildTableRow(TSharedPtr<FChang
 		];
 	}
 
-	return SNew(SBuildDataRow, InOwnerTable, InItem);
+	return SNew(SBuildDataRow, InOwnerTable, InItem).ToolTipText(FText::FromString(InItem->Description));
 }
 
 // Button callbacks
@@ -153,8 +157,7 @@ void SGameSyncTab::Construct(const FArguments& InArgs)
 		SNew(SVerticalBox)
 		// Toolbar at the top of the tab // Todo: Maybe use a FToolBarBuilder instead
 		+SVerticalBox::Slot()
-		.FillHeight(0.05f)
-		// .MaxHeight(35.0f)
+		.AutoHeight()
 		.Padding(10.0f, 5.0f)
 		[
 			SNew(SBorder)
@@ -269,7 +272,7 @@ void SGameSyncTab::Construct(const FArguments& InArgs)
 		// Stream banner
 		+SVerticalBox::Slot()
 		.Padding(0.0f, 5.0f)
-		.FillHeight(0.2f)
+		.AutoHeight()
 		[
 			SNew(SOverlay)
 			+SOverlay::Slot()
@@ -389,7 +392,7 @@ void SGameSyncTab::Construct(const FArguments& InArgs)
 		// Horde builds
 		+SVerticalBox::Slot()
 		.Padding(0.0f, 5.0f)
-		.FillHeight(0.45f)
+		.FillHeight(0.5f)
 		[
 			SAssignNew(HordeBuildsView, SListView<TSharedPtr<FChangeInfo>>)
 			.ListItemsSource(&HordeBuilds)
@@ -404,10 +407,10 @@ void SGameSyncTab::Construct(const FArguments& InArgs)
 				.FixedWidth(35.0f)
 				+SHeaderRow::Column(HordeTableColumnChange)
 				.DefaultLabel(LOCTEXT("HordeHeaderChange", "Change"))
-				.FillWidth(0.1f)
+				.FixedWidth(70.0f)
 				+SHeaderRow::Column(HordeTableColumnTime)
 				.DefaultLabel(LOCTEXT("HordeHeaderTime", "Time"))
-				.FillWidth(0.1f)
+				.FixedWidth(65.0f)
 				+SHeaderRow::Column(HordeTableColumnAuthor)
 				.DefaultLabel(LOCTEXT("HordeHeaderAuthor", "Author"))
 				.FillWidth(0.15f)
@@ -417,7 +420,7 @@ void SGameSyncTab::Construct(const FArguments& InArgs)
 		]
 		// Log
 		+SVerticalBox::Slot()
-		.Padding(0.0f, 5.0f, 0.0f, 10.0f)
+		.Padding(0.0f, 5.0f)
 		.FillHeight(0.3f)
 		[
 			SAssignNew(SyncLog, SLogWidget)
