Made the builds list / log resizeable

diff --git a/Engine/Source/Programs/SlateUGS/Private/Widgets/SGameSyncTab.cpp b/Engine/Source/Programs/SlateUGS/Private/Widgets/SGameSyncTab.cpp
--- a/Engine/Source/Programs/SlateUGS/Private/Widgets/SGameSyncTab.cpp
+++ b/Engine/Source/Programs/SlateUGS/Private/Widgets/SGameSyncTab.cpp
@@ -389,41 +389,45 @@ void SGameSyncTab::Construct(const FArguments& InArgs)
 				]
 			]
 		]
-		// Horde builds
 		+SVerticalBox::Slot()
 		.Padding(0.0f, 5.0f)
-		.FillHeight(0.5f)
 		[
-			SAssignNew(HordeBuildsView, SListView<TSharedPtr<FChangeInfo>>)
-			.ListItemsSource(&HordeBuilds)
-			.SelectionMode(ESelectionMode::Single)
-			.IsEnabled_Lambda([this] { return !Tab->IsSyncing(); })
-			.OnGenerateRow(this, &SGameSyncTab::GenerateHordeBuildTableRow)
-			.OnContextMenuOpening(this, &SGameSyncTab::OnRightClickedBuild)
-			.HeaderRow(
-				SNew(SHeaderRow)
-				+SHeaderRow::Column(HordeTableColumnStatus)
-				.DefaultLabel(LOCTEXT("HordeHeaderStatus", ""))
-				.FixedWidth(35.0f)
-				+SHeaderRow::Column(HordeTableColumnChange)
-				.DefaultLabel(LOCTEXT("HordeHeaderChange", "Change"))
-				.FixedWidth(70.0f)
-				+SHeaderRow::Column(HordeTableColumnTime)
-				.DefaultLabel(LOCTEXT("HordeHeaderTime", "Time"))
-				.FixedWidth(65.0f)
-				+SHeaderRow::Column(HordeTableColumnAuthor)
-				.DefaultLabel(LOCTEXT("HordeHeaderAuthor", "Author"))
-				.FillWidth(0.15f)
-				+SHeaderRow::Column(HordeTableColumnDescription)
-				.DefaultLabel(LOCTEXT("HordeHeaderDescription", "Description"))
-			)
-		]
-		// Log
-		+SVerticalBox::Slot()
-		.Padding(0.0f, 5.0f)
-		.FillHeight(0.3f)
-		[
-			SAssignNew(SyncLog, SLogWidget)
+			SNew(SSplitter)
+			.Orientation(EOrientation::Orient_Vertical)
+			// Horde builds
+			+SSplitter::Slot()
+			.Value(0.65f)
+			[
+				SAssignNew(HordeBuildsView, SListView<TSharedPtr<FChangeInfo>>)
+				.ListItemsSource(&HordeBuilds)
+				.SelectionMode(ESelectionMode::Single)
+				.IsEnabled_Lambda([this] { return !Tab->IsSyncing(); })
+				.OnGenerateRow(this, &SGameSyncTab::GenerateHordeBuildTableRow)
+				.OnContextMenuOpening(this, &SGameSyncTab::OnRightClickedBuild)
+				.HeaderRow(
+					SNew(SHeaderRow)
+					+SHeaderRow::Column(HordeTableColumnStatus)
+					.DefaultLabel(LOCTEXT("HordeHeaderStatus", ""))
+					.FixedWidth(35.0f)
+					+SHeaderRow::Column(HordeTableColumnChange)
+					.DefaultLabel(LOCTEXT("HordeHeaderChange", "Change"))
+					.FixedWidth(70.0f)
+					+SHeaderRow::Column(HordeTableColumnTime)
+					.DefaultLabel(LOCTEXT("HordeHeaderTime", "Time"))
+					.FixedWidth(65.0f)
+					+SHeaderRow::Column(HordeTableColumnAuthor)
+					.DefaultLabel(LOCTEXT("HordeHeaderAuthor", "Author"))
+					.FillWidth(0.15f)
+					+SHeaderRow::Column(HordeTableColumnDescription)
+					.DefaultLabel(LOCTEXT("HordeHeaderDescription", "Description"))
+				)
+			]
+			// Log
+			+SSplitter::Slot()
+			.Value(0.35f)
+			[
+				SAssignNew(SyncLog, SLogWidget)
+			]
 		]
 	];
 }
