Start ScheduledSync window

diff --git a/Engine/Source/Programs/SlateUGS/Private/Widgets/SGameSyncTab.cpp b/Engine/Source/Programs/SlateUGS/Private/Widgets/SGameSyncTab.cpp
--- a/Engine/Source/Programs/SlateUGS/Private/Widgets/SGameSyncTab.cpp
+++ b/Engine/Source/Programs/SlateUGS/Private/Widgets/SGameSyncTab.cpp
@@ -6,6 +6,7 @@
 #include "Framework/MultiBox/MultiBoxBuilder.h"
 
 #include "SPositiveActionButton.h"
+#include "SScheduledSyncWindow.h"
 #include "SSimpleComboButton.h"
 #include "SSimpleButton.h"
 #include "SSyncFilterWindow.h"
@@ -136,6 +137,7 @@ TSharedRef<ITableRow> SGameSyncTab::GenerateHordeBuildTableRow(TSharedPtr<FChang
 TSharedRef<SWidget> SGameSyncTab::MakeSyncButtonDropdown()
 {
 	FMenuBuilder MenuBuilder(true, nullptr);
+
 	MenuBuilder.AddMenuEntry(
 		FText::FromString(TEXT("Sync Latest")),
 		FText::FromString(TEXT("Sync to the latest submitted changelist")),
@@ -143,7 +145,16 @@ TSharedRef<SWidget> SGameSyncTab::MakeSyncButtonDropdown()
 		FUIAction(FExecuteAction::CreateRaw(Tab, &UGSTab::OnSyncLatest)),
 		NAME_None,
 		EUserInterfaceActionType::Button
-		);
+	);
+
+	MenuBuilder.AddMenuEntry(
+		FText::FromString(TEXT("Setup Schedule Sync")),
+		FText::FromString(TEXT("Setup a schedule sync to run at specific time")),
+		FSlateIcon(),
+		FUIAction(FExecuteAction::CreateLambda([this] { FSlateApplication::Get().AddModalWindow(SNew(SScheduledSyncWindow).Tab(Tab), Tab->GetTabArgs().GetOwnerWindow(), false); } )),
+		NAME_None,
+		EUserInterfaceActionType::Button
+	);
 
 	return MenuBuilder.MakeWidget();
 }
diff --git a/Engine/Source/Programs/SlateUGS/Private/Widgets/SScheduledSyncWindow.cpp b/Engine/Source/Programs/SlateUGS/Private/Widgets/SScheduledSyncWindow.cpp
new file mode 100644
--- /dev/null
+++ b/Engine/Source/Programs/SlateUGS/Private/Widgets/SScheduledSyncWindow.cpp
@@ -0,0 +1,78 @@
+// Copyright Epic Games, Inc. All Rights Reserved.
+
+#include "SScheduledSyncWindow.h"
+
+#include "UGSTab.h"
+
+#include "Widgets/Layout/SUniformGridPanel.h"
+#include "SPrimaryButton.h"
+#include "SSimpleButton.h"
+#include "SPopupTextWindow.h"
+
+#define LOCTEXT_NAMESPACE "SScheduledSyncWindow"
+
+void SScheduledSyncWindow::Construct(const FArguments& InArgs)
+{
+	Tab = InArgs._Tab;
+
+	SWindow::Construct(SWindow::FArguments()
+	.Title(LOCTEXT("WindowTitle", "Schedule Sync"))
+	.SizingRule(ESizingRule::FixedSize)
+	.ClientSize(FVector2D(400, 300))
+	[
+		SNew(SVerticalBox)
+		// Hint text
+		// TODO check box to enable ScheduleSync
+		+SVerticalBox::Slot()
+		.AutoHeight()
+		.Padding(20.0f, 20.0f)
+		[
+			SNew(STextBlock)
+			.Text(LOCTEXT("ScheduleSync", "Set a time for a sync to go off on all or some project."))
+		]
+		// TODO widget to enter a time of day
+		// TODO check box to enable it for all projects
+		// TODO check box to enable each project if the previous one is not enabled
+		// Buttons
+		+SVerticalBox::Slot()
+		.AutoHeight()
+		.Padding(0.0f, 20.0f)
+		[
+			SNew(SHorizontalBox)
+			+SHorizontalBox::Slot()
+			.Padding(0.0f, 0.0f, 10.0f, 0.0f)
+			.HAlign(HAlign_Right)
+			[
+				SNew(SUniformGridPanel)
+				.SlotPadding(FMargin(10.0f, 0.0f))
+				+SUniformGridPanel::Slot(0, 0)
+				[
+					SNew(SPrimaryButton)
+					.Text(LOCTEXT("SaveButtonText", "Save"))
+					.OnClicked(this, &SScheduledSyncWindow::OnSaveClicked)
+				]
+				+SUniformGridPanel::Slot(1, 0)
+				[
+					SNew(SButton)
+					.HAlign(HAlign_Center)
+					.Text(LOCTEXT("CancelButtonText", "Cancel"))
+					.OnClicked(this, &SScheduledSyncWindow::OnCancelClicked)
+				]
+			]
+		]
+	]);
+}
+
+FReply SScheduledSyncWindow::OnSaveClicked()
+{
+	FSlateApplication::Get().FindWidgetWindow(AsShared())->RequestDestroyWindow();
+
+	return FReply::Handled();
+}
+
+FReply SScheduledSyncWindow::OnCancelClicked()
+{
+	FSlateApplication::Get().FindWidgetWindow(AsShared())->RequestDestroyWindow();
+
+	return FReply::Handled();
+}
diff --git a/Engine/Source/Programs/SlateUGS/Private/Widgets/SScheduledSyncWindow.h b/Engine/Source/Programs/SlateUGS/Private/Widgets/SScheduledSyncWindow.h
new file mode 100644
--- /dev/null
+++ b/Engine/Source/Programs/SlateUGS/Private/Widgets/SScheduledSyncWindow.h
@@ -0,0 +1,23 @@
+// Copyright Epic Games, Inc. All Rights Reserved.
+
+#pragma once
+
+#include "Widgets/SWindow.h"
+
+class UGSTab;
+
+class SScheduledSyncWindow : public SWindow
+{
+public:
+	SLATE_BEGIN_ARGS(SScheduledSyncWindow) {}
+		SLATE_ARGUMENT(UGSTab*, Tab)
+	SLATE_END_ARGS()
+
+	void Construct(const FArguments& InArgs);
+
+private:
+	FReply OnSaveClicked();
+	FReply OnCancelClicked();
+
+	UGSTab* Tab;
+};
