Added a missing #undef directive

diff --git a/Engine/Source/Programs/SlateUGS/Private/Widgets/SScheduledSyncWindow.cpp b/Engine/Source/Programs/SlateUGS/Private/Widgets/SScheduledSyncWindow.cpp
--- a/Engine/Source/Programs/SlateUGS/Private/Widgets/SScheduledSyncWindow.cpp
+++ b/Engine/Source/Programs/SlateUGS/Private/Widgets/SScheduledSyncWindow.cpp
@@ -76,3 +76,5 @@ FReply SScheduledSyncWindow::OnCancelClicked()
 
 	return FReply::Handled();
 }
+
+#undef LOCTEXT_NAMESPACE
