Made fetching of Perforce executable platform generic, so it now also finds p4 on Mac.

diff --git a/Engine/Source/Programs/SlateUGS/Private/UGSCore/Perforce.cpp b/Engine/Source/Programs/SlateUGS/Private/UGSCore/Perforce.cpp
--- a/Engine/Source/Programs/SlateUGS/Private/UGSCore/Perforce.cpp
+++ b/Engine/Source/Programs/SlateUGS/Private/UGSCore/Perforce.cpp
@@ -41,12 +41,41 @@ FString GetTempFileName()
 
 struct FPerforceExe
 {
+	bool bValidP4 = false;
+	FString ExecutablePath;
+
+	FString GetPerforceExe()
+	{
+		if (!ExecutablePath.IsEmpty())
+		{
+			return ExecutablePath;
+		}
+
 #if PLATFORM_WINDOWS
-	FString PerforceExe = TEXT("C:\\Program Files (x86)\\Perforce\\p4.exe");
+		ExecutablePath = TEXT("C:\\Program Files (x86)\\Perforce\\p4.exe");
 #else
-	FString PerforceExe = TEXT("p4");
-#endif
-	bool bValidP4 = false;
+		int32 ReturnCode;
+		FString OutResults;
+		FString OutErrors;
+
+#if PLATFORM_MAC
+		FPlatformProcess::ExecProcess(TEXT("/usr/bin/mdfind"), TEXT("\"kMDItemFSName = 'p4' && kMDItemContentType = 'public.unix-executable'\""), &ReturnCode, &OutResults, &OutErrors);
+#elif PLATFORM_LINUX
+		FPlatformProcess::ExecProcess(TEXT("which"), TEXT("p4"), &ReturnCode, &OutResults, &OutErrors);
+#endif // PLATFORM_MAC
+
+		if (ReturnCode != 0)
+		{
+			// Todo: Log OutErrors somehow (should this take function the output log?)
+			return FString();
+		}
+
+		ExecutablePath = OutResults.TrimEnd();
+#endif // PLATFORM_WINDOWS
+
+		UE_LOG(LogSlate, Warning, TEXT("Result: %s, Error: %s"), *ExecutablePath, *OutErrors);
+		return ExecutablePath;
+	}
 
 	inline int InnerRunCommand(const FString& CommandLine, TArray<FString>& OutLines)
 	{
@@ -55,7 +84,7 @@ struct FPerforceExe
 		void* WritePipe = nullptr;
 		FPlatformProcess::CreatePipe(ReadPipe, WritePipe);
 
-		FProcHandle P4Proc = FPlatformProcess::CreateProc(*PerforceExe, *CommandLine, false, true, true, &ProcId, 0, nullptr, WritePipe, ReadPipe);
+		FProcHandle P4Proc = FPlatformProcess::CreateProc(*GetPerforceExe(), *CommandLine, false, true, true, &ProcId, 0, nullptr, WritePipe, ReadPipe);
 
 		FString P4Output;
 		FString LatestOutput = FPlatformProcess::ReadPipe(ReadPipe);
@@ -418,9 +447,11 @@ TSharedRef<FPerforceConnection> FPerforceConnection::OpenClient(const FString& N
 bool FPerforceConnection::Info(TSharedPtr<FPerforceInfoRecord>& OutInfo, FEvent* AbortEvent, FOutputDevice& Log) const
 {
 	TArray<TMap<FString, FString>> TagRecords;
+	UE_LOG(LogSlate, Warning, TEXT("Info called"));
 	if(!RunCommand(TEXT("info -s"), TagRecords, ECommandOptions::NoClient, AbortEvent, Log) || TagRecords.Num() != 1)
 	{
 		OutInfo = TSharedPtr<FPerforceInfoRecord>();
+		UE_LOG(LogSlate, Warning, TEXT("Info returned false"));
 		return false;
 	}
 	else
@@ -1037,6 +1068,7 @@ bool FPerforceConnection::RunCommand(const FString& CommandLine, TArray<FPerforc
 	TArray<TMap<FString, FString>> TagRecords;
 	if(!RunCommand(CommandLine, TagRecords, Options, AbortEvent, Log))
 	{
+		UE_LOG(LogSlate, Warning, TEXT("RunCommand1 returned false"));
 		return false;
 	}
 	else
@@ -1077,7 +1109,7 @@ bool FPerforceConnection::RunCommand(const FString& CommandLine, TArray<FString>
 bool FPerforceConnection::RunCommand(const FString& CommandLine, EPerforceOutputChannel Channel, TArray<FString>& OutLines, ECommandOptions Options, FEvent* AbortEvent, FOutputDevice& Log) const
 {
 	FString FullCommandLine = GetFullCommandLine(CommandLine, Options);
-	Log.Logf(TEXT("p4> %s %s"), *FPaths::GetCleanFilename(GPerforceExe.PerforceExe), *FullCommandLine);
+	Log.Logf(TEXT("p4> %s %s"), *FPaths::GetCleanFilename(GPerforceExe.GetPerforceExe()), *FullCommandLine);
 
 	TArray<FString> RawOutputLines;
 
@@ -1107,7 +1139,7 @@ bool FPerforceConnection::RunCommand(const FString& CommandLine, EPerforceOutput
 bool FPerforceConnection::RunCommand(const FString& CommandLine, const TCHAR* Input, TFunction<bool(const FPerforceOutputLine&)> HandleOutput, ECommandOptions Options, FEvent* AbortEvent, FOutputDevice& Log) const
 {
 	FString FullCommandLine = GetFullCommandLine(CommandLine, Options);
-	Log.Logf(TEXT("p4> %s %s"), *FPaths::GetCleanFilename(GPerforceExe.PerforceExe), *FullCommandLine);
+	Log.Logf(TEXT("p4> %s %s"), *FPaths::GetCleanFilename(GPerforceExe.GetPerforceExe()), *FullCommandLine);
 
 	TArray<FString> RawOutputLines;
 
@@ -1369,7 +1401,7 @@ bool FPerforceConnection::RunCommandWithBinaryOutput(const FString& CommandLine,
     void* WritePipe = nullptr;
     FPlatformProcess::CreatePipe(ReadPipe, WritePipe);
 
-	FProcHandle P4Proc = FPlatformProcess::CreateProc(*GPerforceExe.PerforceExe, *FullCommandLine, false, true, true, &ProcId, 0, nullptr, WritePipe, ReadPipe);
+	FProcHandle P4Proc = FPlatformProcess::CreateProc(*GPerforceExe.GetPerforceExe(), *FullCommandLine, false, true, true, &ProcId, 0, nullptr, WritePipe, ReadPipe);
 
 	size_t BufferPos = 0;
 	size_t BufferEnd = 0;
