Removed a debug print statement, oopsies

diff --git a/Engine/Source/Programs/SlateUGS/Private/UGSCore/Perforce.cpp b/Engine/Source/Programs/SlateUGS/Private/UGSCore/Perforce.cpp
--- a/Engine/Source/Programs/SlateUGS/Private/UGSCore/Perforce.cpp
+++ b/Engine/Source/Programs/SlateUGS/Private/UGSCore/Perforce.cpp
@@ -73,7 +73,6 @@ struct FPerforceExe
 		ExecutablePath = OutResults.TrimEnd();
 #endif // PLATFORM_WINDOWS
 
-		UE_LOG(LogSlate, Warning, TEXT("Result: %s, Error: %s"), *ExecutablePath, *OutErrors);
 		return ExecutablePath;
 	}
 
