SlateUGS - Compile fix for Windows

diff --git a/Engine/Source/Programs/SlateUGS/Private/UGSCore/Perforce.cpp b/Engine/Source/Programs/SlateUGS/Private/UGSCore/Perforce.cpp
--- a/Engine/Source/Programs/SlateUGS/Private/UGSCore/Perforce.cpp
+++ b/Engine/Source/Programs/SlateUGS/Private/UGSCore/Perforce.cpp
@@ -125,7 +125,7 @@ struct FPerforceExe
 		TArray<FString> Lines;
 		if (RunCommand(TEXT(""), Lines) == -1)
 		{
-			PerforceExe.ReplaceInline(TEXT(" (x86)"), TEXT(""));
+			ExecutablePath.ReplaceInline(TEXT(" (x86)"), TEXT(""));
 			bValidP4 = (RunCommand(TEXT(""), Lines) == -1);
 		}
 		else
