Add support for Schedule Sync, UI still needs to be setup but the logic works

diff --git a/Engine/Source/Programs/SlateUGS/Private/ScheduledSyncTimer.cpp b/Engine/Source/Programs/SlateUGS/Private/ScheduledSyncTimer.cpp
new file mode 100644
--- /dev/null
+++ b/Engine/Source/Programs/SlateUGS/Private/ScheduledSyncTimer.cpp
@@ -0,0 +1,73 @@
+// Copyright Epic Games, Inc. All Rights Reserved.
+
+#include "ScheduledSyncTimer.h"
+#include "HAL/RunnableThread.h"
+#include "UGSLog.h"
+
+ScheduledSyncTimer::~ScheduledSyncTimer()
+{
+	bFinished = true;
+
+	if (WorkerThread != nullptr)
+	{
+		WorkerThread->WaitForCompletion();
+		WorkerThread = nullptr;
+	}
+}
+
+void ScheduledSyncTimer::Start(const FDateTime& InScheduledTime, const TFunction<void()>& InTimerElapsedCallback)
+{
+	check(WorkerThread == nullptr);
+
+	bFinished            = false;
+	ScheduledTime        = InScheduledTime;
+	TimerElapsedCallback = InTimerElapsedCallback;
+
+	WorkerThread = FRunnableThread::Create(this, TEXT("ScheduledSyncTimer"));
+}
+
+void ScheduledSyncTimer::Stop()
+{
+	bFinished = true;
+
+	if (WorkerThread)
+	{
+		WorkerThread->WaitForCompletion();
+		WorkerThread = nullptr;
+	}
+}
+
+uint32 ScheduledSyncTimer::Run()
+{
+	while (!bFinished)
+	{
+		FDateTime Now = FDateTime::Now();
+
+		if (IsEngineExitRequested())
+		{
+			bFinished = true;
+		}
+
+		if (Now > ScheduledTime)
+		{
+			if (TimerElapsedCallback)
+			{
+				TimerElapsedCallback();
+			}
+
+			// Increment our next sync 1 day from now
+			ScheduledTime += FTimespan(1, 0, 0, 0);
+
+			UE_LOG(LogSlateUGS, Log, TEXT("Schedule: Started ScheduleTimer for %s (%s remaining)"),
+				*ScheduledTime.ToString(TEXT("%Y/%m/%d at %h:%M%a")),
+				*(ScheduledTime - Now).ToString(TEXT("%h hours and %m minutes")));
+		}
+
+		// lets check every second to avoid spinning this thread to hard, while still be responsive when requesting to exit
+		// the ideal option would be to create a wait event tied to the bFinished bool, and sleep for longer chunks *unless* that bool
+		// is flipped, but this is just simple for now
+		FPlatformProcess::Sleep(1.0f);
+	}
+
+	return 0;
+}
diff --git a/Engine/Source/Programs/SlateUGS/Private/ScheduledSyncTimer.h b/Engine/Source/Programs/SlateUGS/Private/ScheduledSyncTimer.h
new file mode 100644
--- /dev/null
+++ b/Engine/Source/Programs/SlateUGS/Private/ScheduledSyncTimer.h
@@ -0,0 +1,24 @@
+// Copyright Epic Games, Inc. All Rights Reserved.
+
+#pragma once
+
+#include "HAL/Runnable.h"
+#include "Templates/Function.h"
+
+class ScheduledSyncTimer : FRunnable
+{
+public:
+	~ScheduledSyncTimer();
+
+	void Start(const FDateTime& InScheduledTime, const TFunction<void()>& InTimerElapsedCallback);
+	void Stop();
+
+private:
+	virtual uint32 Run() override;
+
+	bool bFinished = false;
+	
+	FDateTime ScheduledTime;
+	FRunnableThread* WorkerThread = nullptr;
+	TFunction<void()> TimerElapsedCallback;
+};
diff --git a/Engine/Source/Programs/SlateUGS/Private/UGSTab.cpp b/Engine/Source/Programs/SlateUGS/Private/UGSTab.cpp
--- a/Engine/Source/Programs/SlateUGS/Private/UGSTab.cpp
+++ b/Engine/Source/Programs/SlateUGS/Private/UGSTab.cpp
@@ -28,11 +28,13 @@ UGSTab::UGSTab() : TabArgs(nullptr, FTabId()),
 				   GameSyncTabView(SNew(SGameSyncTab).Tab(this)),
 				   bHasQueuedMessages(false)
 {
-	Initialize();
+	Initialize(nullptr);
 }
 
-void UGSTab::Initialize()
+void UGSTab::Initialize(TSharedPtr<UGSCore::FUserSettings> InUserSettings)
 {
+	UserSettings = InUserSettings;
+
 	TabWidget->SetContent(EmptyTabView);
 	TabWidget->SetLabel(FText(LOCTEXT("TabName", "Select a Project")));
 }
@@ -175,6 +177,11 @@ namespace
 	}
 }
 
+bool UGSTab::IsProjectLoaded() const
+{
+	return !ProjectFileName.IsEmpty();
+}
+
 bool UGSTab::OnWorkspaceChosen(const FString& Project)
 {
 	bool bIsDataValid = FPaths::FileExists(Project); // Todo: Check that the project file is also associated with a workspace
@@ -502,11 +509,6 @@ void UGSTab::SetupWorkspace()
 		return;
 	}
 
-	FString DataFolder = FString(FPlatformProcess::UserSettingsDir()) / TEXT("UnrealGameSync");
-	IFileManager::Get().MakeDirectory(*DataFolder);
-
-	UserSettings = MakeShared<UGSCore::FUserSettings>(*(DataFolder / TEXT("UnrealGameSync_Slate.ini")));
-
 	PerforceClient    = DetectSettings->PerforceClient;
 	WorkspaceSettings = UserSettings->FindOrAddWorkspace(*DetectSettings->BranchClientPath);
 	ProjectSettings   = UserSettings->FindOrAddProject(*DetectSettings->NewSelectedClientFileName);
@@ -524,6 +526,8 @@ void UGSTab::SetupWorkspace()
 		ClientKey = ClientKey.Left(ClientKey.Len() - 1);
 	}
 
+	FString DataFolder = FString(FPlatformProcess::UserSettingsDir()) / TEXT("UnrealGameSync");
+
 	FString ProjectLogBaseName         = DataFolder / FString::Printf(TEXT("%s@%s"), *PerforceClient->ClientName, *ClientKey.Replace(TEXT("/"), TEXT("$")));
 	FString TelemetryProjectIdentifier = UGSCore::FPerforceUtils::GetClientOrDepotDirectoryName(*DetectSettings->NewSelectedProjectIdentifier);
 	FString SelectedProjectIdentifier  = DetectSettings->NewSelectedProjectIdentifier;
diff --git a/Engine/Source/Programs/SlateUGS/Private/UGSTab.h b/Engine/Source/Programs/SlateUGS/Private/UGSTab.h
--- a/Engine/Source/Programs/SlateUGS/Private/UGSTab.h
+++ b/Engine/Source/Programs/SlateUGS/Private/UGSTab.h
@@ -30,7 +30,7 @@ class UGSTab
 {
 public:
 	UGSTab();
-	void Initialize();
+	void Initialize(TSharedPtr<UGSCore::FUserSettings> UserSettings);
 
 	const TSharedRef<SDockTab> GetTabWidget();
 	FSpawnTabArgs GetTabArgs() const;
@@ -40,6 +40,8 @@ public:
 
 	void Tick();
 
+	bool IsProjectLoaded() const;
+
 	// Slate callbacks
 	bool OnWorkspaceChosen(const FString& Project);
 	void OnSyncChangelist(int Changelist);
@@ -62,8 +64,8 @@ public:
 	TArray<UGSCore::FWorkspaceSyncCategory> GetSyncCategories(SyncCategoryType CategoryType) const;
 	TArray<FString> GetSyncViews(SyncCategoryType CategoryType) const;
 	UGSTabManager* GetTabManager();
-private:
 
+private:
 	void OnWorkspaceSyncComplete(
 		TSharedRef<UGSCore::FWorkspaceUpdateContext, ESPMode::ThreadSafe> WorkspaceContext,
 		UGSCore::EWorkspaceUpdateResult SyncResult,
diff --git a/Engine/Source/Programs/SlateUGS/Private/UGSTabManager.cpp b/Engine/Source/Programs/SlateUGS/Private/UGSTabManager.cpp
--- a/Engine/Source/Programs/SlateUGS/Private/UGSTabManager.cpp
+++ b/Engine/Source/Programs/SlateUGS/Private/UGSTabManager.cpp
@@ -4,6 +4,16 @@
 #include "Containers/Array.h"
 #include "UGSLog.h"
 
+#include "UGSCore/UserSettings.h"
+
+UGSTabManager::UGSTabManager()
+{
+	FString DataFolder = FString(FPlatformProcess::UserSettingsDir()) / TEXT("UnrealGameSync");
+	IFileManager::Get().MakeDirectory(*DataFolder);
+
+	UserSettings = MakeShared<UGSCore::FUserSettings>(*(DataFolder / TEXT("UnrealGameSync_Slate.ini")));
+}
+
 void UGSTabManager::ConstructTabs()
 {
 	TSharedRef<FTabManager::FStack> TabStack = FTabManager::NewStack();
@@ -33,6 +43,8 @@ void UGSTabManager::ConstructTabs()
 	);
 
 	TabManager->RestoreFrom(Layout, nullptr);
+
+	SetupScheduledSync();
 }
 
 void UGSTabManager::Tick()
@@ -41,16 +53,81 @@ void UGSTabManager::Tick()
 	{
 		Tab.Tick();
 	}
+
+	if (bScheduledTimerElapsed)
+	{
+		bScheduledTimerElapsed = false;
+
+		ScheduleTimerElapsed();
+	}
 }
 
 TSharedRef<SDockTab> UGSTabManager::SpawnTab(int Index, const FSpawnTabArgs& Args)
 {
-	Tabs[Index].Initialize();
+	Tabs[Index].Initialize(UserSettings);
 	Tabs[Index].SetTabArgs(Args);
 	Tabs[Index].SetTabManager(this);
+
 	return Tabs[Index].GetTabWidget();
 }
 
+void UGSTabManager::SetupScheduledSync()
+{
+	StopScheduledSyncTimer();
+
+	if (UserSettings->bScheduleEnabled)
+	{
+		StartScheduledSyncTimer();
+	}
+}
+
+void UGSTabManager::StartScheduledSyncTimer()
+{
+	bScheduledTimerElapsed = false;
+
+	// We need to grab the current date info for the ticks, as our timer will increment by days, as well as
+	// we will use ::Now() to check if we our timer is ready to signal
+	FDateTime CurrentTime   = FDateTime::Now();
+	FDateTime ScheduledTime = CurrentTime.GetDate() + UserSettings->ScheduleTime;
+
+	// If our CurrentTime is greater then our ScheduleTime, add 1 day to our ScheduleTime to hit the next occurrence of that time
+	if (CurrentTime > ScheduledTime)
+	{
+		ScheduledTime += FTimespan(1, 0, 0, 0);
+	}
+
+	UE_LOG(LogSlateUGS, Log, TEXT("Schedule: Started ScheduleTimer for %s (%s remaining)"),
+		*ScheduledTime.ToString(TEXT("%Y/%m/%d at %h:%M%a")),
+		*(ScheduledTime - CurrentTime).ToString(TEXT("%h hours and %m minutes")));
+
+	SyncTimer.Start(ScheduledTime, [this] {
+		// flip this atomic so on the next tick we can sync on the main thread
+		bScheduledTimerElapsed = true;
+	});
+}
+
+void UGSTabManager::StopScheduledSyncTimer()
+{
+	SyncTimer.Stop();
+}
+
+void UGSTabManager::ScheduleTimerElapsed()
+{
+	TSharedPtr<FGlobalTabmanager> TabManager = FGlobalTabmanager::Get();
+	for (int TabIndex = 0; TabIndex < MaxTabs; TabIndex++)
+	{
+		FName TabId = GetTabId(TabIndex);
+		if (TabManager->FindExistingLiveTab(TabId).IsValid() && Tabs[TabIndex].IsProjectLoaded())
+		{
+			// TODO we need a way to sync based on good vs only latest change (which could be broken)
+			Tabs[TabIndex].OnSyncLatest();
+		}
+	}
+
+	// flip this back for the next sync timer that will happen in ~24 hours if UGS is left open
+	bScheduledTimerElapsed = false;
+}
+
 // Todo: replace this super hacky way of fetching the first available closed tab
 void UGSTabManager::ActivateTab()
 {
@@ -64,7 +141,7 @@ void UGSTabManager::ActivateTab()
 		}
 		if (TabManager->TryInvokeTab(TabId, false).IsValid())
 		{
-			Tabs[TabIndex].Initialize();
+			Tabs[TabIndex].Initialize(UserSettings);
 			return;
 		}
 	}
@@ -75,4 +152,4 @@ void UGSTabManager::ActivateTab()
 FName UGSTabManager::GetTabId(int TabIndex) const
 {
 	return FName(FString(TEXT("UGS Tab: ")) + FString::FromInt(TabIndex));
-}
\ No newline at end of file
+}
diff --git a/Engine/Source/Programs/SlateUGS/Private/UGSTabManager.h b/Engine/Source/Programs/SlateUGS/Private/UGSTabManager.h
--- a/Engine/Source/Programs/SlateUGS/Private/UGSTabManager.h
+++ b/Engine/Source/Programs/SlateUGS/Private/UGSTabManager.h
@@ -3,11 +3,17 @@
 #pragma once
 
 #include "UGSTab.h"
+#include "ScheduledSyncTimer.h"
+
+namespace UGSCore
+{
+	struct FUserSettings;
+}
 
 class UGSTabManager
 {
 public:
-	UGSTabManager() = default;
+	UGSTabManager();
 	void ConstructTabs();
 
 	void Tick();
@@ -17,6 +23,16 @@ public:
 	FName GetTabId(int TabIndex) const;
 
 private:
+	void SetupScheduledSync();
+	void StartScheduledSyncTimer();
+	void StopScheduledSyncTimer();
+	void ScheduleTimerElapsed();
+
 	static constexpr int MaxTabs = 10;
 	TStaticArray<UGSTab, MaxTabs> Tabs;
+
+	TSharedPtr<UGSCore::FUserSettings> UserSettings;
+
+	std::atomic<bool> bScheduledTimerElapsed;
+	ScheduledSyncTimer SyncTimer;
 };
