Remove debug

diff --git a/Engine/Source/Programs/SlateUGS/Private/UGSCore/Perforce.cpp b/Engine/Source/Programs/SlateUGS/Private/UGSCore/Perforce.cpp
--- a/Engine/Source/Programs/SlateUGS/Private/UGSCore/Perforce.cpp
+++ b/Engine/Source/Programs/SlateUGS/Private/UGSCore/Perforce.cpp
@@ -446,11 +446,9 @@ TSharedRef<FPerforceConnection> FPerforceConnection::OpenClient(const FString& N
 bool FPerforceConnection::Info(TSharedPtr<FPerforceInfoRecord>& OutInfo, FEvent* AbortEvent, FOutputDevice& Log) const
 {
 	TArray<TMap<FString, FString>> TagRecords;
-	UE_LOG(LogSlate, Warning, TEXT("Info called"));
 	if(!RunCommand(TEXT("info -s"), TagRecords, ECommandOptions::NoClient, AbortEvent, Log) || TagRecords.Num() != 1)
 	{
 		OutInfo = TSharedPtr<FPerforceInfoRecord>();
-		UE_LOG(LogSlate, Warning, TEXT("Info returned false"));
 		return false;
 	}
 	else
