Fix initialization of the p4 exe path

diff --git a/Engine/Source/Programs/SlateUGS/Private/UGSCore/Perforce.cpp b/Engine/Source/Programs/SlateUGS/Private/UGSCore/Perforce.cpp
--- a/Engine/Source/Programs/SlateUGS/Private/UGSCore/Perforce.cpp
+++ b/Engine/Source/Programs/SlateUGS/Private/UGSCore/Perforce.cpp
@@ -42,19 +42,26 @@ FString GetTempFileName()
 struct FPerforceExe
 {
 	bool bValidP4 = false;
-	FString ExecutablePath;
+
+	// Default assumed locations, we will try to find the real path
+#if PLATFORM_WINDOWS
+	FString ExecutablePath = TEXT("C:\\Program Files (x86)\\Perforce\\p4.exe");
+#else
+	FString ExecutablePath = TEXT("/usr/bin/p4");
+#endif
 
 	FString GetPerforceExe()
 	{
-		if (!ExecutablePath.IsEmpty())
+		static bool bExecuatblePathCached = false;
+		if (bExecuatblePathCached)
 		{
 			return ExecutablePath;
 		}
 
-#if PLATFORM_WINDOWS
-		ExecutablePath = TEXT("C:\\Program Files (x86)\\Perforce\\p4.exe");
-#else
-		int32 ReturnCode;
+		bExecuatblePathCached = true;
+
+#if PLATFORM_MAC || PLATFORM_LINUX
+		int32 ReturnCode = -1;
 		FString OutResults;
 		FString OutErrors;
 
@@ -67,11 +74,12 @@ struct FPerforceExe
 		if (ReturnCode != 0)
 		{
 			// Todo: Log OutErrors somehow (should this take function the output log?)
-			return FString();
+			// for now just return our default paths
+			return ExecutablePath;
 		}
 
 		ExecutablePath = OutResults.TrimEnd();
-#endif // PLATFORM_WINDOWS
+#endif // PLATFORM_MAC || PLATFORM_LINUX
 
 		return ExecutablePath;
 	}
@@ -119,14 +127,14 @@ struct FPerforceExe
 		return -1;
 	}
 
-	FPerforceExe()
+	void VerifyPerforcePath()
 	{
 #if PLATFORM_WINDOWS
-		TArray<FString> Lines;
-		if (RunCommand(TEXT(""), Lines) == -1)
+		bValidP4 = FPaths::FileExists(ExecutablePath);
+		if (!bValidP4)
 		{
 			ExecutablePath.ReplaceInline(TEXT(" (x86)"), TEXT(""));
-			bValidP4 = (RunCommand(TEXT(""), Lines) == -1);
+			bValidP4 = FPaths::FileExists(ExecutablePath);
 		}
 		else
 #endif
@@ -134,6 +142,11 @@ struct FPerforceExe
 			bValidP4 = true;
 		}
 	}
+
+	FPerforceExe()
+	{
+		VerifyPerforcePath();
+	}
 };
 static FPerforceExe GPerforceExe;
 
