Fix spelling

diff --git a/Engine/Source/Programs/SlateUGS/Private/UGSCore/Perforce.cpp b/Engine/Source/Programs/SlateUGS/Private/UGSCore/Perforce.cpp
--- a/Engine/Source/Programs/SlateUGS/Private/UGSCore/Perforce.cpp
+++ b/Engine/Source/Programs/SlateUGS/Private/UGSCore/Perforce.cpp
@@ -52,13 +52,13 @@ struct FPerforceExe
 
 	FString GetPerforceExe()
 	{
-		static bool bExecuatblePathCached = false;
-		if (bExecuatblePathCached)
+		static bool bExecutablePathCached = false;
+		if (bExecutablePathCached)
 		{
 			return ExecutablePath;
 		}
 
-		bExecuatblePathCached = true;
+		bExecutablePathCached = true;
 
 #if PLATFORM_MAC || PLATFORM_LINUX
 		int32 ReturnCode = -1;
