Add basic elements for setting up scheduled sync, the UI needs to be re-worked but the functionality is there

diff --git a/Engine/Source/Programs/SlateUGS/Private/ScheduledSyncTimer.cpp b/Engine/Source/Programs/SlateUGS/Private/ScheduledSyncTimer.cpp
--- a/Engine/Source/Programs/SlateUGS/Private/ScheduledSyncTimer.cpp
+++ b/Engine/Source/Programs/SlateUGS/Private/ScheduledSyncTimer.cpp
@@ -63,10 +63,10 @@ uint32 ScheduledSyncTimer::Run()
 				*(ScheduledTime - Now).ToString(TEXT("%h hours and %m minutes")));
 		}
 
-		// lets check every second to avoid spinning this thread to hard, while still be responsive when requesting to exit
+		// lets check every 100ms to avoid spinning this thread to hard, while still be responsive when requesting to exit
 		// the ideal option would be to create a wait event tied to the bFinished bool, and sleep for longer chunks *unless* that bool
 		// is flipped, but this is just simple for now
-		FPlatformProcess::Sleep(1.0f);
+		FPlatformProcess::Sleep(0.1f);
 	}
 
 	return 0;
diff --git a/Engine/Source/Programs/SlateUGS/Private/UGSTab.cpp b/Engine/Source/Programs/SlateUGS/Private/UGSTab.cpp
--- a/Engine/Source/Programs/SlateUGS/Private/UGSTab.cpp
+++ b/Engine/Source/Programs/SlateUGS/Private/UGSTab.cpp
@@ -469,6 +469,11 @@ UGSTabManager* UGSTab::GetTabManager()
 	return TabManager;
 }
 
+TSharedPtr<UGSCore::FUserSettings> UGSTab::GetUserSettings() const
+{
+	return UserSettings;
+}
+
 // This is getting called on a thread, lock our stuff up
 void UGSTab::OnWorkspaceSyncComplete(TSharedRef<UGSCore::FWorkspaceUpdateContext, ESPMode::ThreadSafe> WorkspaceContext, UGSCore::EWorkspaceUpdateResult SyncResult, const FString& StatusMessage)
 {
diff --git a/Engine/Source/Programs/SlateUGS/Private/UGSTab.h b/Engine/Source/Programs/SlateUGS/Private/UGSTab.h
--- a/Engine/Source/Programs/SlateUGS/Private/UGSTab.h
+++ b/Engine/Source/Programs/SlateUGS/Private/UGSTab.h
@@ -64,6 +64,7 @@ public:
 	TArray<UGSCore::FWorkspaceSyncCategory> GetSyncCategories(SyncCategoryType CategoryType) const;
 	TArray<FString> GetSyncViews(SyncCategoryType CategoryType) const;
 	UGSTabManager* GetTabManager();
+	TSharedPtr<UGSCore::FUserSettings> GetUserSettings() const;
 
 private:
 	void OnWorkspaceSyncComplete(
diff --git a/Engine/Source/Programs/SlateUGS/Private/UGSTabManager.h b/Engine/Source/Programs/SlateUGS/Private/UGSTabManager.h
--- a/Engine/Source/Programs/SlateUGS/Private/UGSTabManager.h
+++ b/Engine/Source/Programs/SlateUGS/Private/UGSTabManager.h
@@ -22,10 +22,12 @@ public:
 	void ActivateTab();
 	FName GetTabId(int TabIndex) const;
 
-private:
+	// Will setup a schedule sync, as well as tear down an existing one and make a new one
 	void SetupScheduledSync();
-	void StartScheduledSyncTimer();
 	void StopScheduledSyncTimer();
+
+private:
+	void StartScheduledSyncTimer();
 	void ScheduleTimerElapsed();
 
 	static constexpr int MaxTabs = 10;
diff --git a/Engine/Source/Programs/SlateUGS/Private/Widgets/SScheduledSyncWindow.cpp b/Engine/Source/Programs/SlateUGS/Private/Widgets/SScheduledSyncWindow.cpp
--- a/Engine/Source/Programs/SlateUGS/Private/Widgets/SScheduledSyncWindow.cpp
+++ b/Engine/Source/Programs/SlateUGS/Private/Widgets/SScheduledSyncWindow.cpp
@@ -2,18 +2,25 @@
 
 #include "SScheduledSyncWindow.h"
 
+#include "UGSLog.h"
 #include "UGSTab.h"
+#include "UGSTabManager.h"
 
+#include "Widgets/Input/SCheckBox.h"
+#include "Widgets/Input/SEditableTextBox.h"
 #include "Widgets/Layout/SUniformGridPanel.h"
 #include "SPrimaryButton.h"
 #include "SSimpleButton.h"
 #include "SPopupTextWindow.h"
 
+#include "UGSCore/UserSettings.h"
+
 #define LOCTEXT_NAMESPACE "SScheduledSyncWindow"
 
 void SScheduledSyncWindow::Construct(const FArguments& InArgs)
 {
 	Tab = InArgs._Tab;
+	UserSettings = Tab->GetUserSettings();
 
 	SWindow::Construct(SWindow::FArguments()
 	.Title(LOCTEXT("WindowTitle", "Schedule Sync"))
@@ -25,11 +32,42 @@ void SScheduledSyncWindow::Construct(const FArguments& InArgs)
 		// TODO check box to enable ScheduleSync
 		+SVerticalBox::Slot()
 		.AutoHeight()
-		.Padding(20.0f, 20.0f)
+		.Padding(10.0f, 5.0f)
 		[
 			SNew(STextBlock)
 			.Text(LOCTEXT("ScheduleSync", "Set a time for a sync to go off on all or some project."))
 		]
+		+ SVerticalBox::Slot()
+		.AutoHeight()  
+		.Padding(10.0f, 0.0f, 0.0f, 0.0f)
+		[
+			SNew(SCheckBox)
+			.ForegroundColor(FSlateColor::UseForeground())
+			.ToolTipText(LOCTEXT("EnableDisableScheudleSyncCheckbox", "Enable or disable Schedule Sync"))
+			.IsChecked(this, &SScheduledSyncWindow::HandleGetScheduleSyncChecked)
+			.OnCheckStateChanged(this, &SScheduledSyncWindow::HandleScheduleSyncChanged)
+			[
+				SNew(SBox)
+				.VAlign(VAlign_Center)
+				.HAlign(HAlign_Center)
+				.Padding(FMargin(4.0, 2.0))
+				[
+					SNew(STextBlock)
+					.Text(LOCTEXT("EnableDisableScheudleSync", "Enable Schedule Sync"))
+				]
+			]
+		]
+
+		+SVerticalBox::Slot()
+		.AutoHeight()
+		.Padding(10.0f, 5.0f)
+		[
+			SNew(SEditableTextBox)
+			.ClearKeyboardFocusOnCommit(false)
+			.OnTextCommitted(this, &SScheduledSyncWindow::HandleTextBoxTextCommited)
+			.SelectAllTextOnCommit(true)
+			.Text(this, &SScheduledSyncWindow::HandleTextBoxText)
+		]
 		// TODO widget to enter a time of day
 		// TODO check box to enable it for all projects
 		// TODO check box to enable each project if the previous one is not enabled
@@ -48,8 +86,8 @@ void SScheduledSyncWindow::Construct(const FArguments& InArgs)
 				+SUniformGridPanel::Slot(0, 0)
 				[
 					SNew(SPrimaryButton)
-					.Text(LOCTEXT("SaveButtonText", "Save"))
-					.OnClicked(this, &SScheduledSyncWindow::OnSaveClicked)
+					.Text(LOCTEXT("OKButtonText", "OK"))
+					.OnClicked(this, &SScheduledSyncWindow::OnOKClicked)
 				]
 				+SUniformGridPanel::Slot(1, 0)
 				[
@@ -61,9 +99,99 @@ void SScheduledSyncWindow::Construct(const FArguments& InArgs)
 			]
 		]
 	]);
+
+	// Lets make the default input our current saved ScheduleTime
+	// Move to an FDateTime from a FTimespan due to FTimespan having an implicit +/- symbol for ToString methods. As well as am/pm strings
+	FDateTime SavedScheduledTime = FDateTime(1, 1, 1, UserSettings->ScheduleTime.GetHours(), UserSettings->ScheduleTime.GetMinutes());
+	Input = SavedScheduledTime.ToString(TEXT("%h:%M%a"));
+}
+
+FText SScheduledSyncWindow::HandleTextBoxText() const
+{
+	return FText::FromString(Input);
+}
+
+void SScheduledSyncWindow::HandleTextBoxTextCommited(const FText& NewText, ETextCommit::Type CommitInfo)
+{
+	FString NewTime = NewText.ToString();
+
+	// Valid format for a yyyy.mm.dd-hh.mm.ss
+	// turn our am/pm into a 1/0 to use to see if we need to add 12 hours when creating the FDateTime::Parse format
+	NewTime.ReplaceInline(TEXT("am"), TEXT(" 1"));
+	NewTime.ReplaceInline(TEXT("pm"), TEXT(" 0"));
+	
+	NewTime.ReplaceInline(TEXT(":"), TEXT(" "));
+
+	// split up on a single delimiter
+	TArray<FString> Tokens;
+	NewTime.ParseIntoArray(Tokens, TEXT(" "), true);
+
+	// make sure it parsed it properly (within reason)
+	// we expect an hour, minute and am/pm
+	if (Tokens.Num() != 3)
+	{
+		UE_LOG(LogSlateUGS, Warning, TEXT("Invalid Time Format %s, should be in the form of <HH:MM[am/pm]>"), *NewText.ToString());
+		return;
+	}
+
+	int32 Hour   = FCString::Atoi(*Tokens[0]);
+	int32 Minute = FCString::Atoi(*Tokens[1]);
+	int32 bIsAm  = FCString::Atoi(*Tokens[2]);
+
+	if (!bIsAm)
+	{
+		Hour += 12;
+	}
+
+	// 1.1.1 is just a valid date which we dont care about but is required for parsing into a FDateTime
+	FDateTime InputTime;
+
+	bInputValid = FDateTime::Parse(FString::Printf(TEXT("1.1.1-%i.%i.0"), Hour, Minute), InputTime);
+
+	if (bInputValid)
+	{
+		FTimespan NewScheduleTime = FTimespan(Hour, Minute, 0);
+		Input = InputTime.ToString(TEXT("%h:%M%a"));
+
+		if (NewScheduleTime != UserSettings->ScheduleTime)
+		{
+			UserSettings->ScheduleTime = NewScheduleTime;
+			UserSettings->Save();
+
+			Tab->GetTabManager()->SetupScheduledSync();
+		}
+	}
+	else
+	{
+		UE_LOG(LogSlateUGS, Warning, TEXT("Invalid Time Format %s, should be in the form of <HH:MM[am/pm]>"), *NewText.ToString());
+	}
+}
+
+void SScheduledSyncWindow::HandleScheduleSyncChanged(ECheckBoxState InCheck)
+{
+	bool bNewValue = InCheck == ECheckBoxState::Checked;
+	if (UserSettings->bScheduleEnabled != bNewValue)  
+	{
+		UserSettings->bScheduleEnabled = bNewValue;
+		UserSettings->Save();
+
+		if (UserSettings->bScheduleEnabled)
+		{
+			Tab->GetTabManager()->SetupScheduledSync();
+		}
+		else
+		{
+			Tab->GetTabManager()->StopScheduledSyncTimer();
+		}
+	}
+}
+
+ECheckBoxState SScheduledSyncWindow::HandleGetScheduleSyncChecked() const
+{
+	return UserSettings->bScheduleEnabled ? ECheckBoxState::Checked : ECheckBoxState::Unchecked;
 }
 
-FReply SScheduledSyncWindow::OnSaveClicked()
+FReply SScheduledSyncWindow::OnOKClicked()
 {
 	FSlateApplication::Get().FindWidgetWindow(AsShared())->RequestDestroyWindow();
 
diff --git a/Engine/Source/Programs/SlateUGS/Private/Widgets/SScheduledSyncWindow.h b/Engine/Source/Programs/SlateUGS/Private/Widgets/SScheduledSyncWindow.h
--- a/Engine/Source/Programs/SlateUGS/Private/Widgets/SScheduledSyncWindow.h
+++ b/Engine/Source/Programs/SlateUGS/Private/Widgets/SScheduledSyncWindow.h
@@ -3,9 +3,15 @@
 #pragma once
 
 #include "Widgets/SWindow.h"
+#include "Styling/SlateColor.h"
 
 class UGSTab;
 
+namespace UGSCore
+{
+	struct FUserSettings;
+}
+
 class SScheduledSyncWindow : public SWindow
 {
 public:
@@ -16,8 +22,25 @@ public:
 	void Construct(const FArguments& InArgs);
 
 private:
-	FReply OnSaveClicked();
+	/** Handles getting the text color of the editable text box. */
+	FSlateColor HandleTextBoxForegroundColor() const;
+
+	/** Handles getting the text to be displayed in the editable text box. */ 
+	FText HandleTextBoxText() const;
+
+	/** Handles committing the text in the editable text box. */
+	void HandleTextBoxTextCommited(const FText& NewText, ETextCommit::Type CommitInfo);
+
+	/** Handle the check box enable or disabling Schedule Sync */
+	void HandleScheduleSyncChanged(ECheckBoxState InCheck);
+	ECheckBoxState HandleGetScheduleSyncChecked() const;
+
+	FReply OnOKClicked();
 	FReply OnCancelClicked();
 
+	bool bInputValid = false;
+	FString Input;
+
 	UGSTab* Tab;
+	TSharedPtr<UGSCore::FUserSettings> UserSettings;
 };
