Add SettingsWindow as well as AfterSync options for now

diff --git a/Engine/Source/Programs/SlateUGS/Private/UGSTab.cpp b/Engine/Source/Programs/SlateUGS/Private/UGSTab.cpp
--- a/Engine/Source/Programs/SlateUGS/Private/UGSTab.cpp
+++ b/Engine/Source/Programs/SlateUGS/Private/UGSTab.cpp
@@ -204,6 +204,29 @@ bool UGSTab::OnWorkspaceChosen(const FString& Project)
 
 void UGSTab::OnSyncChangelist(int Changelist)
 {
+	// Options on what to do with workspace when updating it
+	UGSCore::EWorkspaceUpdateOptions Options = UGSCore::EWorkspaceUpdateOptions::Sync | UGSCore::EWorkspaceUpdateOptions::SyncArchives | UGSCore::EWorkspaceUpdateOptions::GenerateProjectFiles;
+	if (UserSettings->bAutoResolveConflicts)
+	{
+		Options |= UGSCore::EWorkspaceUpdateOptions::AutoResolveChanges;
+	}
+	if (UserSettings->bUseIncrementalBuilds)
+	{
+		Options |= UGSCore::EWorkspaceUpdateOptions::UseIncrementalBuilds;
+	}
+	if (UserSettings->bBuildAfterSync)
+	{
+		Options |= UGSCore::EWorkspaceUpdateOptions::Build;
+	}
+	if (UserSettings->bBuildAfterSync && UserSettings->bRunAfterSync)
+	{
+		Options |= UGSCore::EWorkspaceUpdateOptions::RunAfterSync;
+	}
+	if (UserSettings->bOpenSolutionAfterSync)
+	{
+		Options |= UGSCore::EWorkspaceUpdateOptions::OpenSolutionAfterSync;
+	}
+
 	TSharedRef<UGSCore::FWorkspaceUpdateContext, ESPMode::ThreadSafe> Context = MakeShared<UGSCore::FWorkspaceUpdateContext, ESPMode::ThreadSafe>(
 		Changelist,
 		Options,
@@ -564,29 +587,6 @@ void UGSTab::SetupWorkspace()
 		WorkspaceSettings->SyncView,
 		WorkspaceSettings->SyncExcludedCategories);
 
-	// Options on what to do with workspace when updating it
-	Options = UGSCore::EWorkspaceUpdateOptions::Sync | UGSCore::EWorkspaceUpdateOptions::SyncArchives | UGSCore::EWorkspaceUpdateOptions::GenerateProjectFiles;
-	if (UserSettings->bAutoResolveConflicts)
-	{
-		Options |= UGSCore::EWorkspaceUpdateOptions::AutoResolveChanges;
-	}
-	if (UserSettings->bUseIncrementalBuilds)
-	{
-		Options |= UGSCore::EWorkspaceUpdateOptions::UseIncrementalBuilds;
-	}
-	if (UserSettings->bBuildAfterSync)
-	{
-		Options |= UGSCore::EWorkspaceUpdateOptions::Build;
-	}
-	if (UserSettings->bBuildAfterSync && UserSettings->bRunAfterSync)
-	{
-		Options |= UGSCore::EWorkspaceUpdateOptions::RunAfterSync;
-	}
-	if (UserSettings->bOpenSolutionAfterSync)
-	{
-		Options |= UGSCore::EWorkspaceUpdateOptions::OpenSolutionAfterSync;
-	}
-
 	// Setup our Perforce and Event monitoring threads
 	FString BranchClientPath = DetectSettings->BranchDirectoryName;
 	FString SelectedClientFileName = DetectSettings->NewSelectedClientFileName;
diff --git a/Engine/Source/Programs/SlateUGS/Private/UGSTab.h b/Engine/Source/Programs/SlateUGS/Private/UGSTab.h
--- a/Engine/Source/Programs/SlateUGS/Private/UGSTab.h
+++ b/Engine/Source/Programs/SlateUGS/Private/UGSTab.h
@@ -96,7 +96,6 @@ private:
 	TSharedPtr<UGSCore::FPerforceConnection> PerforceClient;
 	TSharedPtr<UGSCore::FUserWorkspaceSettings> WorkspaceSettings;
 	TSharedPtr<UGSCore::FUserProjectSettings> ProjectSettings;
-	UGSCore::EWorkspaceUpdateOptions Options;
 	TSharedPtr<UGSCore::FDetectProjectSettingsTask> DetectSettings;
 	TArray<FString> CombinedSyncFilter;
 	TSharedPtr<UGSCore::FUserSettings> UserSettings;
diff --git a/Engine/Source/Programs/SlateUGS/Private/Widgets/SGameSyncTab.cpp b/Engine/Source/Programs/SlateUGS/Private/Widgets/SGameSyncTab.cpp
--- a/Engine/Source/Programs/SlateUGS/Private/Widgets/SGameSyncTab.cpp
+++ b/Engine/Source/Programs/SlateUGS/Private/Widgets/SGameSyncTab.cpp
@@ -1,6 +1,8 @@
 // Copyright Epic Games, Inc. All Rights Reserved.
 
 #include "SGameSyncTab.h"
+#include "SSyncFilterWindow.h"
+#include "SSettingsWindow.h"
 
 #include "Framework/Application/SlateApplication.h"
 #include "Framework/MultiBox/MultiBoxBuilder.h"
@@ -9,7 +11,6 @@
 #include "SScheduledSyncWindow.h"
 #include "SSimpleComboButton.h"
 #include "SSimpleButton.h"
-#include "SSyncFilterWindow.h"
 #include "Widgets/SLogWidget.h"
 #include "Widgets/Layout/SHeader.h"
 #include "Widgets/Testing/STestSuite.h"
@@ -275,7 +276,7 @@ void SGameSyncTab::Construct(const FArguments& InArgs)
 						SNew(SSimpleButton)
 						.Text(LOCTEXT("Settings", "Settings"))
 						.Icon(FAppStyle::Get().GetBrush("Icons.Settings"))
-						.IsEnabled(false) // Todo: enable after adding this functionality
+						.OnClicked_Lambda([this] { FSlateApplication::Get().AddModalWindow(SNew(SSettingsWindow).Tab(Tab), Tab->GetTabArgs().GetOwnerWindow(), false); return FReply::Handled(); })
 					]
 				]
 			]
diff --git a/Engine/Source/Programs/SlateUGS/Private/Widgets/SScheduledSyncWindow.cpp b/Engine/Source/Programs/SlateUGS/Private/Widgets/SScheduledSyncWindow.cpp
--- a/Engine/Source/Programs/SlateUGS/Private/Widgets/SScheduledSyncWindow.cpp
+++ b/Engine/Source/Programs/SlateUGS/Private/Widgets/SScheduledSyncWindow.cpp
@@ -86,8 +86,8 @@ void SScheduledSyncWindow::Construct(const FArguments& InArgs)
 				+SUniformGridPanel::Slot(0, 0)
 				[
 					SNew(SPrimaryButton)
-					.Text(LOCTEXT("OKButtonText", "OK"))
-					.OnClicked(this, &SScheduledSyncWindow::OnOKClicked)
+					.Text(LOCTEXT("OkButtonText", "Ok"))
+					.OnClicked(this, &SScheduledSyncWindow::OnOkClicked)
 				]
 				+SUniformGridPanel::Slot(1, 0)
 				[
@@ -191,7 +191,7 @@ ECheckBoxState SScheduledSyncWindow::HandleGetScheduleSyncChecked() const
 	return UserSettings->bScheduleEnabled ? ECheckBoxState::Checked : ECheckBoxState::Unchecked;
 }
 
-FReply SScheduledSyncWindow::OnOKClicked()
+FReply SScheduledSyncWindow::OnOkClicked()
 {
 	FSlateApplication::Get().FindWidgetWindow(AsShared())->RequestDestroyWindow();
 
diff --git a/Engine/Source/Programs/SlateUGS/Private/Widgets/SScheduledSyncWindow.h b/Engine/Source/Programs/SlateUGS/Private/Widgets/SScheduledSyncWindow.h
--- a/Engine/Source/Programs/SlateUGS/Private/Widgets/SScheduledSyncWindow.h
+++ b/Engine/Source/Programs/SlateUGS/Private/Widgets/SScheduledSyncWindow.h
@@ -35,7 +35,7 @@ private:
 	void HandleScheduleSyncChanged(ECheckBoxState InCheck);
 	ECheckBoxState HandleGetScheduleSyncChecked() const;
 
-	FReply OnOKClicked();
+	FReply OnOkClicked();
 	FReply OnCancelClicked();
 
 	bool bInputValid = false;
diff --git a/Engine/Source/Programs/SlateUGS/Private/Widgets/SSettingsWindow.cpp b/Engine/Source/Programs/SlateUGS/Private/Widgets/SSettingsWindow.cpp
new file mode 100644
--- /dev/null
+++ b/Engine/Source/Programs/SlateUGS/Private/Widgets/SSettingsWindow.cpp
@@ -0,0 +1,220 @@
+// Copyright Epic Games, Inc. All Rights Reserved.
+
+#include "SSettingsWindow.h"
+#include "UGSTab.h"
+
+#include "UGSCore/UserSettings.h"
+
+#include "Widgets/Input/SCheckBox.h"
+#include "Widgets/Layout/SHeader.h"
+#include "Widgets/Layout/SUniformGridPanel.h"
+#include "SPrimaryButton.h"
+
+#define LOCTEXT_NAMESPACE "SSettingsWindow"
+
+void SSettingsWindow::Construct(const FArguments& InArgs)
+{
+	Tab = InArgs._Tab;
+	UserSettings = Tab->GetUserSettings();
+
+	SWindow::Construct(SWindow::FArguments()
+	.Title(LOCTEXT("WindowTitle", "Settings"))
+	.SizingRule(ESizingRule::FixedSize)
+	.ClientSize(FVector2D(1100, 800))
+	[
+		SNew(SVerticalBox)
+		// Hint text
+		+SVerticalBox::Slot()
+		.AutoHeight()
+		.Padding(20.0f, 20.0f)
+		[
+			SNew(STextBlock)
+			.Text(LOCTEXT("SettingsHintText", "Settings for Unreal Game Sync"))
+		]
+		// Start of Settings
+		+SVerticalBox::Slot()
+		.Padding(20.0f, 0.0f)
+		[
+			SNew(SVerticalBox)
+			// After Sync Settings
+			+SVerticalBox::Slot()
+			.FillHeight(0.15f)
+			[
+				SNew(SVerticalBox)
+				+SVerticalBox::Slot()
+				.VAlign(VAlign_Top)
+				.Padding(0.0f, 10.0f)
+				[
+					SNew(SHeader)
+					[
+						SNew(STextBlock)
+						.Text(LOCTEXT("AfterSync", "After Sync Settings"))
+					]
+				]
+				+SVerticalBox::Slot()
+				[
+					SNew(SHorizontalBox)
+					+SHorizontalBox::Slot()
+					[
+						SNew(SCheckBox)
+						.ForegroundColor(FSlateColor::UseForeground())
+						.IsChecked(this, &SSettingsWindow::HandleGetBuildChecked)
+						.OnCheckStateChanged(this, &SSettingsWindow::HandleBuildChanged)
+						[
+							SNew(SBox)
+							.VAlign(VAlign_Center)
+							.HAlign(HAlign_Center)
+							.Padding(FMargin(4.0, 2.0))
+							[
+								SNew(STextBlock)
+								.Text(LOCTEXT("EnableBuildAfterSync", "Build"))
+							]
+						]
+					]
+					+SHorizontalBox::Slot()
+					[
+						SNew(SCheckBox)
+						.ForegroundColor(FSlateColor::UseForeground())
+						.IsChecked(this, &SSettingsWindow::HandleGetRunChecked)
+						.OnCheckStateChanged(this, &SSettingsWindow::HandleRunChanged)
+						[
+							SNew(SBox)
+							.VAlign(VAlign_Center)
+							.HAlign(HAlign_Center)
+							.Padding(FMargin(4.0, 2.0))
+							[
+								SNew(STextBlock)
+								.Text(LOCTEXT("EnableRunAfterSync", "Run"))
+							]
+						]
+					]
+					+SHorizontalBox::Slot()
+					[
+						SNew(SCheckBox)
+						.ForegroundColor(FSlateColor::UseForeground())
+						.IsChecked(this, &SSettingsWindow::HandleGetOpenSolutionChecked)
+						.OnCheckStateChanged(this, &SSettingsWindow::HandleOpenSolutionChanged)
+						[
+							SNew(SBox)
+							.VAlign(VAlign_Center)
+							.HAlign(HAlign_Center)
+							.Padding(FMargin(4.0, 2.0))
+							[
+								SNew(STextBlock)
+								.Text(LOCTEXT("EnableOpenSolutionAfterSync", "Open Solution"))
+							]
+						]
+					]
+				]
+				// Other
+				+SVerticalBox::Slot()
+				.FillHeight(0.6f)
+				[
+					SNew(SVerticalBox)
+					+SVerticalBox::Slot()
+					.VAlign(VAlign_Top)
+					.AutoHeight()
+					.Padding(0.0f, 10.0f)
+					[
+						SNew(SHeader)
+						[
+							SNew(STextBlock)
+							.Text(LOCTEXT("Other Settings", "Other Settings"))
+						]
+					]
+				]
+			]
+		]
+		// Buttons
+		+SVerticalBox::Slot()
+		.AutoHeight()
+		.Padding(0.0f, 20.0f)
+		[
+			SNew(SHorizontalBox)
+			+SHorizontalBox::Slot()
+			.Padding(0.0f, 0.0f, 10.0f, 0.0f)
+			.HAlign(HAlign_Right)
+			[
+				SNew(SUniformGridPanel)
+				.SlotPadding(FMargin(10.0f, 0.0f))
+				+SUniformGridPanel::Slot(0, 0)
+				[
+					SNew(SPrimaryButton)
+					.Text(LOCTEXT("OkButtonText", "Ok"))
+					.OnClicked(this, &SSettingsWindow::OnOkClicked)
+				]
+				+SUniformGridPanel::Slot(1, 0)
+				[
+					SNew(SButton)
+					.HAlign(HAlign_Center)
+					.Text(LOCTEXT("CancelButtonText", "Cancel"))
+					.OnClicked(this, &SSettingsWindow::OnCancelClicked)
+				]
+			]
+		]
+	]);
+}
+
+void SSettingsWindow::HandleBuildChanged(ECheckBoxState InCheck)
+{
+	bool bNewValue = InCheck == ECheckBoxState::Checked;
+
+	if (UserSettings->bBuildAfterSync != bNewValue)
+	{
+		UserSettings->bBuildAfterSync = bNewValue;
+		UserSettings->Save();
+	}
+}
+
+ECheckBoxState SSettingsWindow::HandleGetBuildChecked() const
+{
+	return UserSettings->bBuildAfterSync ? ECheckBoxState::Checked : ECheckBoxState::Unchecked;
+}
+
+void SSettingsWindow::HandleRunChanged(ECheckBoxState InCheck)
+{
+	bool bNewValue = InCheck == ECheckBoxState::Checked;
+
+	if (UserSettings->bRunAfterSync != bNewValue)
+	{
+		UserSettings->bRunAfterSync = bNewValue;
+		UserSettings->Save();
+	}
+}
+
+ECheckBoxState SSettingsWindow::HandleGetRunChecked() const
+{
+	return UserSettings->bRunAfterSync ? ECheckBoxState::Checked : ECheckBoxState::Unchecked;
+}
+
+void SSettingsWindow::HandleOpenSolutionChanged(ECheckBoxState InCheck)
+{
+	bool bNewValue = InCheck == ECheckBoxState::Checked;
+
+	if (UserSettings->bOpenSolutionAfterSync != bNewValue)
+	{
+		UserSettings->bOpenSolutionAfterSync = bNewValue;
+		UserSettings->Save();
+	}
+}
+
+ECheckBoxState SSettingsWindow::HandleGetOpenSolutionChecked() const
+{
+	return UserSettings->bOpenSolutionAfterSync ? ECheckBoxState::Checked : ECheckBoxState::Unchecked;
+}
+
+FReply SSettingsWindow::OnOkClicked()
+{
+	FSlateApplication::Get().FindWidgetWindow(AsShared())->RequestDestroyWindow();
+
+	return FReply::Handled();
+}
+
+FReply SSettingsWindow::OnCancelClicked()
+{
+	FSlateApplication::Get().FindWidgetWindow(AsShared())->RequestDestroyWindow();
+
+	return FReply::Handled();
+}
+
+#undef LOCTEXT_NAMESPACE
diff --git a/Engine/Source/Programs/SlateUGS/Private/Widgets/SSettingsWindow.h b/Engine/Source/Programs/SlateUGS/Private/Widgets/SSettingsWindow.h
new file mode 100644
--- /dev/null
+++ b/Engine/Source/Programs/SlateUGS/Private/Widgets/SSettingsWindow.h
@@ -0,0 +1,41 @@
+// Copyright Epic Games, Inc. All Rights Reserved.
+
+#pragma once
+
+#include "Widgets/SWindow.h"
+
+#include "SCheckBoxList.h"
+
+class UGSTab;
+
+namespace UGSCore
+{
+	struct FUserSettings;
+}
+
+class SSettingsWindow : public SWindow
+{
+public:
+	SLATE_BEGIN_ARGS(SSettingsWindow) {}
+		SLATE_ARGUMENT(UGSTab*, Tab)
+	SLATE_END_ARGS()
+
+	void Construct(const FArguments& InArgs);
+
+private:
+	/** Handle the check box enable or disabling for AfterSync settings */
+	void HandleBuildChanged(ECheckBoxState InCheck);
+	ECheckBoxState HandleGetBuildChecked() const;
+
+	void HandleRunChanged(ECheckBoxState InCheck);
+	ECheckBoxState HandleGetRunChecked() const;
+
+	void HandleOpenSolutionChanged(ECheckBoxState InCheck);
+	ECheckBoxState HandleGetOpenSolutionChecked() const;
+
+	FReply OnOkClicked();
+	FReply OnCancelClicked();
+
+	UGSTab* Tab;
+	TSharedPtr<UGSCore::FUserSettings> UserSettings;
+};
