Highlight the currenty synced CL, if there is one

diff --git a/Engine/Source/Programs/SlateUGS/Private/ChangeInfo.h b/Engine/Source/Programs/SlateUGS/Private/ChangeInfo.h
--- a/Engine/Source/Programs/SlateUGS/Private/ChangeInfo.h
+++ b/Engine/Source/Programs/SlateUGS/Private/ChangeInfo.h
@@ -11,6 +11,7 @@ struct FChangeInfo
 
 	UGSCore::EReviewVerdict ReviewStatus = UGSCore::EReviewVerdict::Unknown;
 	int Changelist = 0;
+	bool bCurrentlySynced = false;
 	FText Author;
 	FString Description;
 
diff --git a/Engine/Source/Programs/SlateUGS/Private/UGSTab.cpp b/Engine/Source/Programs/SlateUGS/Private/UGSTab.cpp
--- a/Engine/Source/Programs/SlateUGS/Private/UGSTab.cpp
+++ b/Engine/Source/Programs/SlateUGS/Private/UGSTab.cpp
@@ -312,6 +312,13 @@ void UGSTab::Tick()
 
 		MessageQueue.Empty();
 	}
+
+	if (bNeedUpdateGameTabBuildList)
+	{
+		bNeedUpdateGameTabBuildList = false;
+
+		UpdateGameTabBuildList();
+	}
 }
 
 namespace
@@ -421,12 +428,15 @@ void UGSTab::UpdateGameTabBuildList()
 					ChangeInfos.Add(HeaderRow);
 				}
 
+				bool bCurrentSyncedChange = Change.Number == WorkspaceSettings->CurrentChangeNumber;
+
 				ChangeInfos.Add(MakeShareable(new FChangeInfo
 				{
 					DisplayTime,
 					false,
 					Status,
 					Change.Number,
+					bCurrentSyncedChange,
 					FText::FromString(FormatUserName(Change.User)),
 					Change.Description
 				}));
@@ -516,6 +526,9 @@ void UGSTab::OnWorkspaceSyncComplete(TSharedRef<UGSCore::FWorkspaceUpdateContext
 	});
 
 	UserSettings->Save();
+
+	// TODO hacky, but allows us to highlight which CL we have synced to in the list
+	bNeedUpdateGameTabBuildList = true;
 }
 
 void UGSTab::SetupWorkspace()
diff --git a/Engine/Source/Programs/SlateUGS/Private/UGSTab.h b/Engine/Source/Programs/SlateUGS/Private/UGSTab.h
--- a/Engine/Source/Programs/SlateUGS/Private/UGSTab.h
+++ b/Engine/Source/Programs/SlateUGS/Private/UGSTab.h
@@ -107,4 +107,6 @@ private:
 	// Queue for handling callbacks from multiple threads that need to be called on the main thread
 	TArray<TFunction<void()>> MessageQueue;
 	std::atomic<bool> bHasQueuedMessages;
+
+	bool bNeedUpdateGameTabBuildList = false;
 };
diff --git a/Engine/Source/Programs/SlateUGS/Private/Widgets/SGameSyncTab.cpp b/Engine/Source/Programs/SlateUGS/Private/Widgets/SGameSyncTab.cpp
--- a/Engine/Source/Programs/SlateUGS/Private/Widgets/SGameSyncTab.cpp
+++ b/Engine/Source/Programs/SlateUGS/Private/Widgets/SGameSyncTab.cpp
@@ -94,6 +94,15 @@ TSharedRef<SWidget> SBuildDataRow::GenerateWidgetForColumn(const FName& ColumnId
 		TextItem->SetText(FText::FromString(Description));
 	}
 
+	if (CurrentItem->bCurrentlySynced)
+	{
+		// Lets make the font white and bold when we are the currently synced CL
+		FSlateFontInfo FontInfo = FAppStyle::Get().GetFontStyle("NormalFontBold");
+
+		TextItem->SetFont(FontInfo);
+		TextItem->SetColorAndOpacity(FLinearColor::White);
+	}
+
 	return SNew(SBox)
 		.VAlign(VAlign_Center)
 		.Padding(5.0f, 2.5f)
@@ -423,10 +432,10 @@ void SGameSyncTab::Construct(const FArguments& InArgs)
 					.FixedWidth(35.0f)
 					+SHeaderRow::Column(HordeTableColumnChange)
 					.DefaultLabel(LOCTEXT("HordeHeaderChange", "Change"))
-					.FixedWidth(70.0f)
+					.FixedWidth(75.0f)
 					+SHeaderRow::Column(HordeTableColumnTime)
 					.DefaultLabel(LOCTEXT("HordeHeaderTime", "Time"))
-					.FixedWidth(65.0f)
+					.FixedWidth(70.0f)
 					+SHeaderRow::Column(HordeTableColumnAuthor)
 					.DefaultLabel(LOCTEXT("HordeHeaderAuthor", "Author"))
 					.FillWidth(0.15f)
