Can cut off some text depending on screen size, add a little more room

diff --git a/Engine/Source/Programs/SlateUGS/Private/Widgets/SModalTaskWindow.cpp b/Engine/Source/Programs/SlateUGS/Private/Widgets/SModalTaskWindow.cpp
--- a/Engine/Source/Programs/SlateUGS/Private/Widgets/SModalTaskWindow.cpp
+++ b/Engine/Source/Programs/SlateUGS/Private/Widgets/SModalTaskWindow.cpp
@@ -35,7 +35,7 @@ void SModalTaskWindow::Construct(const FArguments& InArgs)
 		.HasCloseButton(false)
 		.SupportsMinimize(false)
 		.SupportsMaximize(false)
-		.MaxHeight(140.0f)
+		.MaxHeight(160.0f)
 		[
 			SNew(SVerticalBox)
 			+SVerticalBox::Slot()
