Fix an active sync while closing the application calling back into a possible dead tab

overall this likely *should* have been cleaned up in the dtor of the workspace but being explicit here resolves the crash. Need to dig into more why not being cleaned up correctly

diff --git a/Engine/Source/Programs/SlateUGS/Private/UGSTab.cpp b/Engine/Source/Programs/SlateUGS/Private/UGSTab.cpp
--- a/Engine/Source/Programs/SlateUGS/Private/UGSTab.cpp
+++ b/Engine/Source/Programs/SlateUGS/Private/UGSTab.cpp
@@ -32,6 +32,14 @@ UGSTab::UGSTab() : TabArgs(nullptr, FTabId()),
 	Initialize(nullptr);
 }
 
+UGSTab::~UGSTab()
+{
+	if (Workspace)
+	{
+		Workspace->CancelUpdate();
+	}
+}
+
 void UGSTab::Initialize(TSharedPtr<UGSCore::FUserSettings> InUserSettings)
 {
 	UserSettings = InUserSettings;
diff --git a/Engine/Source/Programs/SlateUGS/Private/UGSTab.h b/Engine/Source/Programs/SlateUGS/Private/UGSTab.h
--- a/Engine/Source/Programs/SlateUGS/Private/UGSTab.h
+++ b/Engine/Source/Programs/SlateUGS/Private/UGSTab.h
@@ -30,6 +30,8 @@ class UGSTab
 {
 public:
 	UGSTab();
+	~UGSTab();
+
 	void Initialize(TSharedPtr<UGSCore::FUserSettings> UserSettings);
 
 	const TSharedRef<SDockTab> GetTabWidget();
