If we fail to setup workspace, dont try to access shared pointers created by it

diff --git a/Engine/Source/Programs/SlateUGS/Private/UGSTab.cpp b/Engine/Source/Programs/SlateUGS/Private/UGSTab.cpp
--- a/Engine/Source/Programs/SlateUGS/Private/UGSTab.cpp
+++ b/Engine/Source/Programs/SlateUGS/Private/UGSTab.cpp
@@ -198,14 +198,16 @@ bool UGSTab::OnWorkspaceChosen(const FString& Project)
 	{
 		ProjectFileName = Project;
 
-		SetupWorkspace();
-		GameSyncTabView->SetStreamPathText(FText::FromString(DetectSettings->StreamName));
-		GameSyncTabView->SetChangelistText(WorkspaceSettings->CurrentChangeNumber);
-		GameSyncTabView->SetProjectPathText(FText::FromString(ProjectFileName));
-		TabWidget->SetContent(GameSyncTabView);
-		TabWidget->SetLabel(FText::FromString(DetectSettings->StreamName));
+		if (SetupWorkspace())
+		{
+			GameSyncTabView->SetStreamPathText(FText::FromString(DetectSettings->StreamName));
+			GameSyncTabView->SetChangelistText(WorkspaceSettings->CurrentChangeNumber);
+			GameSyncTabView->SetProjectPathText(FText::FromString(ProjectFileName));
+			TabWidget->SetContent(GameSyncTabView);
+			TabWidget->SetLabel(FText::FromString(DetectSettings->StreamName));
 
-		return true;
+			return true;
+		}
 	}
 
 	return false;
@@ -553,7 +555,7 @@ void UGSTab::OnWorkspaceSyncComplete(TSharedRef<UGSCore::FWorkspaceUpdateContext
 	UserSettings->Save();
 }
 
-void UGSTab::SetupWorkspace()
+bool UGSTab::SetupWorkspace()
 {
 	ProjectFileName = UGSCore::FUtility::GetPathWithCorrectCase(ProjectFileName);
 
@@ -569,7 +571,7 @@ void UGSTab::SetupWorkspace()
 	if (Result->Failed())
 	{
 		FMessageDialog::Open(EAppMsgType::Ok, Result->GetMessage());
-		return;
+		return false;
 	}
 
 	PerforceClient    = DetectSettings->PerforceClient;
@@ -642,6 +644,8 @@ void UGSTab::SetupWorkspace()
 	// Start the threads
 	PerforceMonitor->Start();
 	EventMonitor->Start();
+
+	return true;
 }
 
 #undef LOCTEXT_NAMESPACE
diff --git a/Engine/Source/Programs/SlateUGS/Private/UGSTab.h b/Engine/Source/Programs/SlateUGS/Private/UGSTab.h
--- a/Engine/Source/Programs/SlateUGS/Private/UGSTab.h
+++ b/Engine/Source/Programs/SlateUGS/Private/UGSTab.h
@@ -83,7 +83,7 @@ private:
 	void UpdateGameTabBuildList();
 
 	// Core functions
-	void SetupWorkspace();
+	bool SetupWorkspace();
 
 	FCriticalSection CriticalSection;
 
