Add a refresh button that will ask PerforceMonitor to refresh its CL list

diff --git a/Engine/Source/Programs/SlateUGS/Private/UGSTab.cpp b/Engine/Source/Programs/SlateUGS/Private/UGSTab.cpp
--- a/Engine/Source/Programs/SlateUGS/Private/UGSTab.cpp
+++ b/Engine/Source/Programs/SlateUGS/Private/UGSTab.cpp
@@ -213,6 +213,11 @@ bool UGSTab::OnWorkspaceChosen(const FString& Project)
 	return false;
 }
 
+void UGSTab::RefreshBuildList()
+{
+	PerforceMonitor->Refresh();
+}
+
 void UGSTab::CancelSync()
 {
 	if (Workspace->IsBusy())
diff --git a/Engine/Source/Programs/SlateUGS/Private/UGSTab.h b/Engine/Source/Programs/SlateUGS/Private/UGSTab.h
--- a/Engine/Source/Programs/SlateUGS/Private/UGSTab.h
+++ b/Engine/Source/Programs/SlateUGS/Private/UGSTab.h
@@ -68,6 +68,7 @@ public:
 	UGSTabManager* GetTabManager();
 	TSharedPtr<UGSCore::FUserSettings> GetUserSettings() const;
 
+	void RefreshBuildList();
 	void CancelSync();
 
 private:
diff --git a/Engine/Source/Programs/SlateUGS/Private/Widgets/SGameSyncTab.cpp b/Engine/Source/Programs/SlateUGS/Private/Widgets/SGameSyncTab.cpp
--- a/Engine/Source/Programs/SlateUGS/Private/Widgets/SGameSyncTab.cpp
+++ b/Engine/Source/Programs/SlateUGS/Private/Widgets/SGameSyncTab.cpp
@@ -225,10 +225,19 @@ void SGameSyncTab::Construct(const FArguments& InArgs)
 					SNew(SHorizontalBox)
 					+SHorizontalBox::Slot()
 					.AutoWidth()
+					[
+						SNew(SSimpleButton)
+						.Icon(FSlateUGSStyle::Get().GetBrush("Icons.Refresh"))
+						.ToolTipText(LOCTEXT("RefreshBuildList", "Refreshes the Build List"))
+						.OnClicked_Lambda([this] { Tab->RefreshBuildList(); return FReply::Handled(); })
+						.IsEnabled_Lambda([this] { return !Tab->IsSyncing(); })
+					]
+					+SHorizontalBox::Slot()
+					.AutoWidth()
 					[
 						SNew(SSimpleComboButton)
 						.Text(LOCTEXT("Sync", "Sync"))
-						.Icon(FSlateUGSStyle::Get().GetBrush("Icons.Refresh"))
+						.Icon(FSlateUGSStyle::Get().GetBrush("Icons.CircleArrowDown"))
 						.HasDownArrow(true)
 						.MenuContent()
 						[
@@ -241,6 +250,7 @@ void SGameSyncTab::Construct(const FArguments& InArgs)
 					[
 						SNew(SSimpleButton)
 						.Text(LOCTEXT("CancelSync", "Cancel"))
+						.ToolTipText(LOCTEXT("CancelSync", "Cancels the current sync/build operation"))
 						.Icon(FSlateUGSStyle::Get().GetBrush("Icons.X"))
 						.OnClicked_Lambda([this] { Tab->CancelSync(); return FReply::Handled(); })
 						.IsEnabled_Lambda([this] { return Tab->IsSyncing(); })
