Implement support for Build button
Add icons for missing ones

diff --git a/Engine/Source/Programs/SlateUGS/Private/UGSTab.cpp b/Engine/Source/Programs/SlateUGS/Private/UGSTab.cpp
--- a/Engine/Source/Programs/SlateUGS/Private/UGSTab.cpp
+++ b/Engine/Source/Programs/SlateUGS/Private/UGSTab.cpp
@@ -264,6 +264,23 @@ void UGSTab::OnSyncChangelist(int Changelist)
 	Workspace->Update(Context);
 }
 
+void UGSTab::OnBuildWorkspace()
+{
+	UGSCore::EWorkspaceUpdateOptions Options = UGSCore::EWorkspaceUpdateOptions::Build;
+
+	TSharedRef<UGSCore::FWorkspaceUpdateContext, ESPMode::ThreadSafe> Context = MakeShared<UGSCore::FWorkspaceUpdateContext, ESPMode::ThreadSafe>(
+		-1,
+		Options,
+		CombinedSyncFilter,
+		GetDefaultBuildStepObjects(DetectSettings->NewProjectEditorTarget, UserSettings),
+		ProjectSettings->BuildSteps,
+		TSet<FGuid>(),
+		GetWorkspaceVariables(DetectSettings));
+
+	// Update the workspace with the Context!
+	Workspace->Update(Context);
+}
+
 void UGSTab::OnViewInSwarmClicked(int Changelist) const
 {
 	FString SwarmURL = FString::Printf(TEXT("https://p4-swarm.epicgames.net/changes/%i"), Changelist);
diff --git a/Engine/Source/Programs/SlateUGS/Private/UGSTab.h b/Engine/Source/Programs/SlateUGS/Private/UGSTab.h
--- a/Engine/Source/Programs/SlateUGS/Private/UGSTab.h
+++ b/Engine/Source/Programs/SlateUGS/Private/UGSTab.h
@@ -57,6 +57,7 @@ public:
 		const TArray<FGuid>& SyncExcludedCategoriesCurrent,
 		const TArray<FString>& SyncViewAll,
 		const TArray<FGuid>& SyncExcludedCategoriesAll);
+	void OnBuildWorkspace();
 
 	// Accessors
 	bool IsSyncing() const;
diff --git a/Engine/Source/Programs/SlateUGS/Private/Widgets/SGameSyncTab.cpp b/Engine/Source/Programs/SlateUGS/Private/Widgets/SGameSyncTab.cpp
--- a/Engine/Source/Programs/SlateUGS/Private/Widgets/SGameSyncTab.cpp
+++ b/Engine/Source/Programs/SlateUGS/Private/Widgets/SGameSyncTab.cpp
@@ -268,7 +268,7 @@ void SGameSyncTab::Construct(const FArguments& InArgs)
 						SNew(SSimpleButton)
 						.Text(LOCTEXT("Build", "Build"))
 						.IsEnabled_Lambda([this] { return !Tab->IsSyncing(); })
-						.IsEnabled(false) // Todo: enable after adding this functionality
+						.OnClicked_Lambda([this] { Tab->OnBuildWorkspace(); return FReply::Handled(); })
 					]
 					+SHorizontalBox::Slot()
 					.AutoWidth()
@@ -285,7 +285,7 @@ void SGameSyncTab::Construct(const FArguments& InArgs)
 					[
 						SNew(SSimpleButton)
 						.Text(LOCTEXT("OpenSolution", "OpenSolution"))
-						.Icon(FSlateUGSStyle::Get().GetBrush("MainFrame.OpenVisualStudio")) // Todo: shouldn't use this icon (repurposing, also could use other IDEs)
+						.Icon(FSlateUGSStyle::Get().GetBrush("Icons.FolderOpen")) // Todo: shouldn't use this icon (repurposing, also could use other IDEs)
 						.IsEnabled_Lambda([this] { return !Tab->IsSyncing(); })
 						.IsEnabled(false) // Todo: enable after adding this functionality
 					]
@@ -322,7 +322,7 @@ void SGameSyncTab::Construct(const FArguments& InArgs)
 					[
 						SNew(SSimpleButton)
 						.Text(LOCTEXT("CleanSolution", "Clean Solution"))
-						.Icon(FSlateUGSStyle::Get().GetBrush("GraphEditor.Clean")) // Todo: shouldn't use this icon (repurposing)
+						.Icon(FSlateUGSStyle::Get().GetBrush("Icons.Delete")) // Todo: shouldn't use this icon (repurposing)
 						.IsEnabled_Lambda([this] { return !Tab->IsSyncing(); })
 						.IsEnabled(false) // Todo: enable after adding this functionality
 					]
