Add suport for Open in Editor and Open Explorer buttons

diff --git a/Engine/Source/Programs/SlateUGS/Private/UGSTab.cpp b/Engine/Source/Programs/SlateUGS/Private/UGSTab.cpp
--- a/Engine/Source/Programs/SlateUGS/Private/UGSTab.cpp
+++ b/Engine/Source/Programs/SlateUGS/Private/UGSTab.cpp
@@ -331,6 +331,20 @@ void UGSTab::OnSyncFilterWindowSaved(
 	UserSettings->Save();
 }
 
+void UGSTab::OnOpenExplorer()
+{
+	FPlatformProcess::ExploreFolder(*FPaths::GetPath(ProjectFileName));
+}
+
+void UGSTab::OnOpenEditor()
+{
+	UGSCore::EBuildConfig EditorBuildConfig = GetEditorBuildConfig();
+
+	FString EditorPath = GetEditorExePath(EditorBuildConfig, DetectSettings);
+
+	EditorProcessHandle = FPlatformProcess::CreateProc(*EditorPath, *ProjectFileName, true, false, false, nullptr, 0, nullptr, nullptr, nullptr);
+}
+
 void UGSTab::QueueMessageForMainThread(TFunction<void()> Function)
 {
 	FScopeLock Lock(&CriticalSection);
@@ -360,6 +374,16 @@ void UGSTab::Tick()
 
 		UpdateGameTabBuildList();
 	}
+
+	// if we have spawned an editor process, check if its closed to reap it
+	if (EditorProcessHandle.IsValid())
+	{
+		if (!FPlatformProcess::IsProcRunning(EditorProcessHandle))
+		{
+			FPlatformProcess::CloseProc(EditorProcessHandle);
+			EditorProcessHandle.Reset();
+		}
+	}
 }
 
 namespace
diff --git a/Engine/Source/Programs/SlateUGS/Private/UGSTab.h b/Engine/Source/Programs/SlateUGS/Private/UGSTab.h
--- a/Engine/Source/Programs/SlateUGS/Private/UGSTab.h
+++ b/Engine/Source/Programs/SlateUGS/Private/UGSTab.h
@@ -58,6 +58,8 @@ public:
 		const TArray<FString>& SyncViewAll,
 		const TArray<FGuid>& SyncExcludedCategoriesAll);
 	void OnBuildWorkspace();
+	void OnOpenExplorer();
+	void OnOpenEditor();
 
 	// Accessors
 	bool IsSyncing() const;
@@ -115,4 +117,7 @@ private:
 	std::atomic<bool> bHasQueuedMessages;
 
 	std::atomic<bool> bNeedUpdateGameTabBuildList;
+
+	// If we launch the editor, keep track of it to reap when its exited
+	FProcHandle EditorProcessHandle;
 };
diff --git a/Engine/Source/Programs/SlateUGS/Private/Widgets/SGameSyncTab.cpp b/Engine/Source/Programs/SlateUGS/Private/Widgets/SGameSyncTab.cpp
--- a/Engine/Source/Programs/SlateUGS/Private/Widgets/SGameSyncTab.cpp
+++ b/Engine/Source/Programs/SlateUGS/Private/Widgets/SGameSyncTab.cpp
@@ -273,21 +273,20 @@ void SGameSyncTab::Construct(const FArguments& InArgs)
 					+SHorizontalBox::Slot()
 					.AutoWidth()
 					[
-						SNew(SSimpleComboButton)
+						SNew(SSimpleButton)
 						.Text(LOCTEXT("RunUnrealEditor", "Run Unreal Editor"))
 						.Icon(FSlateUGSStyle::Get().GetBrush("Icons.Launch"))
-						.HasDownArrow(true)
 						.IsEnabled_Lambda([this] { return !Tab->IsSyncing(); })
-						.IsEnabled(false) // Todo: enable after adding this functionality
+						.OnClicked_Lambda([this] { Tab->OnOpenEditor(); return FReply::Handled(); })
 					]
 					+SHorizontalBox::Slot()
 					.AutoWidth()
 					[
 						SNew(SSimpleButton)
-						.Text(LOCTEXT("OpenSolution", "OpenSolution"))
+						.Text(LOCTEXT("OpenExplorer", "Open Explorer"))
 						.Icon(FSlateUGSStyle::Get().GetBrush("Icons.FolderOpen")) // Todo: shouldn't use this icon (repurposing, also could use other IDEs)
 						.IsEnabled_Lambda([this] { return !Tab->IsSyncing(); })
-						.IsEnabled(false) // Todo: enable after adding this functionality
+						.OnClicked_Lambda([this] { Tab->OnOpenExplorer(); return FReply::Handled(); })
 					]
 				]
 				+SHorizontalBox::Slot()
