Make path absolute if the FileDialog returns a relative path

diff --git a/Engine/Source/Programs/SlateUGS/Private/Widgets/SWorkspaceWindow.cpp b/Engine/Source/Programs/SlateUGS/Private/Widgets/SWorkspaceWindow.cpp
--- a/Engine/Source/Programs/SlateUGS/Private/Widgets/SWorkspaceWindow.cpp
+++ b/Engine/Source/Programs/SlateUGS/Private/Widgets/SWorkspaceWindow.cpp
@@ -244,11 +244,17 @@ FReply SWorkspaceWindow::OnBrowseClicked()
 			EFileDialogFlags::None,
 			OutOpenFilenames
 		);
-	}
 
-	if (!OutOpenFilenames.IsEmpty())
-	{
-		LocalFileText->SetText(FText::FromString(OutOpenFilenames[0]));
+		FString Path = OutOpenFilenames[0];
+		if (FPaths::IsRelative(Path))
+		{
+			Path = FPaths::ConvertRelativePathToFull(Path);
+		}
+
+		if (!OutOpenFilenames.IsEmpty())
+		{
+			LocalFileText->SetText(FText::FromString(Path));
+		}
 	}
 
 	return FReply::Handled();
