Grab a handle to the Log so we can log some basic actions that are not instance
Giving *some* feedback that we are trying to do something

diff --git a/Engine/Source/Programs/SlateUGS/Private/UGSTab.cpp b/Engine/Source/Programs/SlateUGS/Private/UGSTab.cpp
--- a/Engine/Source/Programs/SlateUGS/Private/UGSTab.cpp
+++ b/Engine/Source/Programs/SlateUGS/Private/UGSTab.cpp
@@ -215,6 +215,7 @@ bool UGSTab::OnWorkspaceChosen(const FString& Project)
 
 void UGSTab::RefreshBuildList()
 {
+	Log->Logf(TEXT("Refreshing build list..."));
 	PerforceMonitor->Refresh();
 }
 
@@ -222,6 +223,7 @@ void UGSTab::CancelSync()
 {
 	if (Workspace->IsBusy())
 	{
+		Log->Logf(TEXT("Canceling sync/build..."));
 		Workspace->CancelUpdate();
 	}
 }
@@ -285,21 +287,25 @@ void UGSTab::OnViewInSwarmClicked(int Changelist) const
 {
 	FString SwarmURL = FString::Printf(TEXT("https://p4-swarm.epicgames.net/changes/%i"), Changelist);
 
+	Log->Logf(TEXT("Opening Swarm to %s"), *SwarmURL);
 	FPlatformProcess::LaunchURL(*SwarmURL, nullptr, nullptr);
 }
 
 void UGSTab::OnCopyChangeListClicked(int Changelist) const
 {
+	Log->Logf(TEXT("Copying CL %i to clipboard"), Changelist);
 	FPlatformApplicationMisc::ClipboardCopy(*FString::FromInt(Changelist));
 }
 
 void UGSTab::OnMoreInfoClicked(int Changelist) const
 {
+	Log->Logf(TEXT("Opening CL %i in p4v"), Changelist);
 	PerforceClient->OpenP4VC(FString::Printf(TEXT("change %i"), Changelist));
 }
 
 void UGSTab::OnOpenPerforceClicked() const
 {
+	Log->Logf(TEXT("Opening %s in p4v"), *ProjectFileName);
 	PerforceClient->OpenP4V();
 }
 
@@ -308,7 +314,7 @@ void UGSTab::OnSyncLatest()
 	int ChangeNumber = -1;
 	FEvent* AbortEvent = FPlatformProcess::GetSynchEventFromPool(true);
 
-	if (PerforceClient->LatestChangeList(ChangeNumber, AbortEvent, MakeShared<FLogWidgetTextWriter>(GameSyncTabView->GetSyncLog().ToSharedRef()).Get()))
+	if (PerforceClient->LatestChangeList(ChangeNumber, AbortEvent, *Log.Get()))
 	{
 		OnSyncChangelist(ChangeNumber);
 	}
@@ -333,6 +339,7 @@ void UGSTab::OnSyncFilterWindowSaved(
 
 void UGSTab::OnOpenExplorer()
 {
+	Log->Logf(TEXT("Opening %s in the explorer"), *FPaths::GetPath(ProjectFileName));
 	FPlatformProcess::ExploreFolder(*FPaths::GetPath(ProjectFileName));
 }
 
@@ -342,6 +349,7 @@ void UGSTab::OnOpenEditor()
 
 	FString EditorPath = GetEditorExePath(EditorBuildConfig, DetectSettings);
 
+	Log->Logf(TEXT("Running Editor %s %s"), *EditorPath, *ProjectFileName);
 	EditorProcessHandle = FPlatformProcess::CreateProc(*EditorPath, *ProjectFileName, true, false, false, nullptr, 0, nullptr, nullptr, nullptr);
 }
 
@@ -646,6 +654,7 @@ bool UGSTab::SetupWorkspace()
 	FString LogFileName = DataFolder / FPaths::GetPath(ProjectFileName) + TEXT(".sync.log");
 	GameSyncTabView->SetSyncLogLocation(LogFileName); // Todo: if SetSyncLogLocation fails, then it failed to create a log file and may need to handle that
 	GameSyncTabView->GetSyncLog()->AppendLine(TEXT("Creating log at: ") + LogFileName);
+	Log = MakeShared<FLogWidgetTextWriter>(GameSyncTabView->GetSyncLog().ToSharedRef());
 
 	Workspace = MakeShared<UGSCore::FWorkspace>(
 		PerforceClient.ToSharedRef(),
@@ -657,7 +666,7 @@ bool UGSTab::SetupWorkspace()
 		CurrentChangeNumber,
 		WorkspaceSettings->LastBuiltChangeNumber,
 		TelemetryProjectIdentifier,
-		MakeShared<FLogWidgetTextWriter>(GameSyncTabView->GetSyncLog().ToSharedRef()));
+		Log.ToSharedRef());
 
 	Workspace->OnUpdateComplete = [this] (TSharedRef<UGSCore::FWorkspaceUpdateContext, ESPMode::ThreadSafe> WorkspaceContext, UGSCore::EWorkspaceUpdateResult SyncResult, const FString& StatusMessage) {
 		OnWorkspaceSyncComplete(WorkspaceContext, SyncResult, StatusMessage);
diff --git a/Engine/Source/Programs/SlateUGS/Private/UGSTab.h b/Engine/Source/Programs/SlateUGS/Private/UGSTab.h
--- a/Engine/Source/Programs/SlateUGS/Private/UGSTab.h
+++ b/Engine/Source/Programs/SlateUGS/Private/UGSTab.h
@@ -19,6 +19,7 @@ namespace UGSCore
 }
 
 class UGSTabManager;
+class FLogWidgetTextWriter;
 
 enum SyncCategoryType
 {
@@ -107,6 +108,7 @@ private:
 	TSharedPtr<UGSCore::FDetectProjectSettingsTask> DetectSettings;
 	TArray<FString> CombinedSyncFilter;
 	TSharedPtr<UGSCore::FUserSettings> UserSettings;
+	TSharedPtr<FLogWidgetTextWriter> Log;
 
 	// Monitoring threads
 	TSharedPtr<UGSCore::FPerforceMonitor> PerforceMonitor;
