Fix Windows path not using *.exe

diff --git a/Engine/Source/Programs/SlateUGS/Private/UGSTab.cpp b/Engine/Source/Programs/SlateUGS/Private/UGSTab.cpp
--- a/Engine/Source/Programs/SlateUGS/Private/UGSTab.cpp
+++ b/Engine/Source/Programs/SlateUGS/Private/UGSTab.cpp
@@ -132,7 +132,11 @@ namespace
 
 	FString GetEditorExePath(UGSCore::EBuildConfig Config, TSharedPtr<UGSCore::FDetectProjectSettingsTask> DetectSettings)
 	{
+	#if PLATFORM_WINDOWS
+		FString ExeFileName = TEXT("UnrealEditor.exe");
+	#else
 		FString ExeFileName = TEXT("UnrealEditor");
+	#endif
 
 		if (Config != UGSCore::EBuildConfig::DebugGame && Config != UGSCore::EBuildConfig::Development)
 		{
