Add setting to enable/disable syncing pre-compiled editor/binaries

diff --git a/Engine/Source/Programs/SlateUGS/Private/Widgets/SSettingsWindow.cpp b/Engine/Source/Programs/SlateUGS/Private/Widgets/SSettingsWindow.cpp
--- a/Engine/Source/Programs/SlateUGS/Private/Widgets/SSettingsWindow.cpp
+++ b/Engine/Source/Programs/SlateUGS/Private/Widgets/SSettingsWindow.cpp
@@ -130,6 +130,31 @@ void SSettingsWindow::Construct(const FArguments& InArgs)
 							.Text(LOCTEXT("Other Settings", "Other Settings"))
 						]
 					]
+					+SVerticalBox::Slot()
+					.AutoHeight()
+					.VAlign(VAlign_Top)
+					.HAlign(HAlign_Center)
+					.Padding(0.0f, 10.0f)
+					[
+						SNew(SHorizontalBox)
+						+SHorizontalBox::Slot()
+						[
+							SNew(SCheckBox)
+							.ForegroundColor(FSlateColor::UseForeground())
+							.IsChecked(this, &SSettingsWindow::HandleGetSyncCompiledEditor)
+							.OnCheckStateChanged(this, &SSettingsWindow::HandleSyncCompiledEditor)
+							[
+								SNew(SBox)
+								.VAlign(VAlign_Center)
+								.HAlign(HAlign_Center)
+								.Padding(FMargin(4.0, 2.0))
+								[
+									SNew(STextBlock)
+									.Text(LOCTEXT("EnableSyncPrecompiledEditor", "Sync Precompiled Editor"))
+								]
+							]
+						]
+					]
 				]
 			]
 		]
@@ -211,6 +236,22 @@ ECheckBoxState SSettingsWindow::HandleGetOpenSolutionChecked() const
 	return UserSettings->bOpenSolutionAfterSync ? ECheckBoxState::Checked : ECheckBoxState::Unchecked;
 }
 
+void SSettingsWindow::HandleSyncCompiledEditor(ECheckBoxState InCheck)
+{
+	bool bNewValue = InCheck == ECheckBoxState::Checked;
+
+	if (UserSettings->bSyncPrecompiledEditor != bNewValue)
+	{
+		UserSettings->bSyncPrecompiledEditor = bNewValue;
+		UserSettings->Save();
+	}
+}
+
+ECheckBoxState SSettingsWindow::HandleGetSyncCompiledEditor() const
+{
+	return UserSettings->bSyncPrecompiledEditor ? ECheckBoxState::Checked : ECheckBoxState::Unchecked;
+}
+
 FReply SSettingsWindow::OnOkClicked()
 {
 	FSlateApplication::Get().FindWidgetWindow(AsShared())->RequestDestroyWindow();
diff --git a/Engine/Source/Programs/SlateUGS/Private/Widgets/SSettingsWindow.h b/Engine/Source/Programs/SlateUGS/Private/Widgets/SSettingsWindow.h
--- a/Engine/Source/Programs/SlateUGS/Private/Widgets/SSettingsWindow.h
+++ b/Engine/Source/Programs/SlateUGS/Private/Widgets/SSettingsWindow.h
@@ -33,6 +33,9 @@ private:
 	void HandleOpenSolutionChanged(ECheckBoxState InCheck);
 	ECheckBoxState HandleGetOpenSolutionChecked() const;
 
+	void HandleSyncCompiledEditor(ECheckBoxState InCheck);
+	ECheckBoxState HandleGetSyncCompiledEditor() const;
+
 	FReply OnOkClicked();
 	FReply OnCancelClicked();
 
