SlateUGS: First pass at syncing precompiled editor/binaries

Still work to be done, but we are syncing the pre compiled editor, just need to unzip it!

diff --git a/Engine/Source/Programs/SlateUGS/Private/UGSTab.cpp b/Engine/Source/Programs/SlateUGS/Private/UGSTab.cpp
--- a/Engine/Source/Programs/SlateUGS/Private/UGSTab.cpp
+++ b/Engine/Source/Programs/SlateUGS/Private/UGSTab.cpp
@@ -9,7 +9,6 @@
 #include "UGSCore/Utility.h"
 #include "UGSCore/BuildStep.h"
 #include "UGSCore/DetectProjectSettingsTask.h"
-#include "UGSCore/UserSettings.h"
 #include "UGSCore/PerforceMonitor.h"
 #include "UGSCore/EventMonitor.h"
 
@@ -88,48 +87,6 @@ namespace
 		}
 	};
 
-	// TODO move most of these to a class with access to the settings
-	//bool ShouldSyncPrecompiledEditor(TSharedPtr<FUserSettings> Settings)
-	bool ShouldSyncPrecompiledEditor()
-	{
-		return false;
-		//return Settings->bSyncPrecompiledEditor;// && PerforceMonitor->HasZippedBinaries();
-	}
-
-	// Honestly ... seems ... super hacky/hardcoded. With out all these you assert when trying to merge build targets sooo a bit odd
-	// TODO Need to do each of these ... per ... platform??
-	TMap<FGuid, UGSCore::FCustomConfigObject> GetDefaultBuildStepObjects(const FString& EditorTargetName, TSharedPtr<UGSCore::FUserSettings> Settings)
-	{
-		TArray<UGSCore::FBuildStep> DefaultBuildSteps;
-		DefaultBuildSteps.Add(UGSCore::FBuildStep(FGuid(0x01F66060, 0x73FA4CC8, 0x9CB3E217, 0xFBBA954E), 0, TEXT("Compile UnrealHeaderTool"), TEXT("Compiling UnrealHeaderTool..."), 1, TEXT("UnrealHeaderTool"), HostPlatform, TEXT("Development"), TEXT(""), !ShouldSyncPrecompiledEditor()));
-		FString ActualEditorTargetName = (EditorTargetName.Len() > 0) ? EditorTargetName : "UnrealEditor";
-		DefaultBuildSteps.Add(UGSCore::FBuildStep(FGuid(0xF097FF61, 0xC9164058, 0x839135B4, 0x6C3173D5), 1, FString::Printf(TEXT("Compile %s"), *ActualEditorTargetName), FString::Printf(TEXT("Compiling %s..."), *ActualEditorTargetName), 10, ActualEditorTargetName, HostPlatform, UGSCore::ToString(Settings->CompiledEditorBuildConfig), TEXT(""), !ShouldSyncPrecompiledEditor()));
-		DefaultBuildSteps.Add(UGSCore::FBuildStep(FGuid(0xC6E633A1, 0x956F4AD3, 0xBC956D06, 0xD131E7B4), 2, TEXT("Compile ShaderCompileWorker"), TEXT("Compiling ShaderCompileWorker..."), 1, TEXT("ShaderCompileWorker"), HostPlatform, TEXT("Development"), TEXT(""), !ShouldSyncPrecompiledEditor()));
-		DefaultBuildSteps.Add(UGSCore::FBuildStep(FGuid(0x24FFD88C, 0x79014899, 0x9696AE10, 0x66B4B6E8), 3, TEXT("Compile UnrealLightmass"), TEXT("Compiling UnrealLightmass..."), 1, TEXT("UnrealLightmass"), HostPlatform, TEXT("Development"), TEXT(""), !ShouldSyncPrecompiledEditor()));
-		DefaultBuildSteps.Add(UGSCore::FBuildStep(FGuid(0xFFF20379, 0x06BF4205, 0x8A3EC534, 0x27736688), 4, TEXT("Compile CrashReportClient"), TEXT("Compiling CrashReportClient..."), 1, TEXT("CrashReportClient"), HostPlatform, TEXT("Development"), TEXT(""), !ShouldSyncPrecompiledEditor()));
-		DefaultBuildSteps.Add(UGSCore::FBuildStep(FGuid(0x89FE8A79, 0xD2594C7B, 0xBFB468F7, 0x218B91C2), 5, TEXT("Compile UnrealInsights"), TEXT("Compiling UnrealInsights..."), 1, TEXT("UnrealInsights"), HostPlatform, TEXT("Development"), TEXT(""), !ShouldSyncPrecompiledEditor()));
-		DefaultBuildSteps.Add(UGSCore::FBuildStep(FGuid(0x46312669, 0x5069428D, 0x8D72C241, 0x6C5A322E), 6, TEXT("Launch UnrealInsights"), TEXT("Running UnrealInsights..."), 1, TEXT("UnrealInsights"), HostPlatform, TEXT("Shipping"), TEXT(""), !ShouldSyncPrecompiledEditor()));
-		DefaultBuildSteps.Add(UGSCore::FBuildStep(FGuid(0xBB48CA5B, 0x56824432, 0x824DC451, 0x336A6523), 7, TEXT("Compile Zen Dashboard"), TEXT("Compile ZenDashboard Step..."), 1, TEXT("ZenDashboard"), HostPlatform, TEXT("Development"), TEXT(""), !ShouldSyncPrecompiledEditor()));
-		DefaultBuildSteps.Add(UGSCore::FBuildStep(FGuid(0x586CC0D3, 0x39144DF9, 0xACB62C02, 0xCD9D4FC6), 8, TEXT("Launch Zen Dashboard"), TEXT("Running Zen Dashboard..."), 1, TEXT("ZenDashboard"), HostPlatform, TEXT("Development"), TEXT(""), !ShouldSyncPrecompiledEditor()));
-		DefaultBuildSteps.Add(UGSCore::FBuildStep(FGuid(0x91C2A429, 0xC39149B4, 0x92A54E6B, 0xE71E0F00), 9, TEXT("Compile SwitchboardListener"), TEXT("Compiling SwitchboardListener..."), 1, TEXT("SwitchboardListener"), HostPlatform, TEXT("Development"), TEXT(""), !ShouldSyncPrecompiledEditor()));
-		DefaultBuildSteps.Add(UGSCore::FBuildStep(FGuid(0x5036C75B, 0x8DF04329, 0x82A1869D, 0xD2D48605), 10, TEXT("Compile UnrealMultiUserServer"), TEXT("Compiling UnrealMultiUserServer..."), 1, TEXT("UnrealMultiUserServer"), HostPlatform, TEXT("Development"), TEXT(""), !ShouldSyncPrecompiledEditor()));
-		DefaultBuildSteps.Add(UGSCore::FBuildStep(FGuid(0x274B89C3, 0x9DC64465, 0xA50840AB, 0xC4593CC2), 11, TEXT("Compile UnrealMultiUserSlateServer"), TEXT("Compiling UnrealMultiUserSlateServer..."), 1, TEXT("UnrealMultiUserSlateServer"), HostPlatform, TEXT("Development"), TEXT(""), !ShouldSyncPrecompiledEditor()));
-
-		TMap<FGuid, UGSCore::FCustomConfigObject> DefaultBuildStepObjects;
-		for (const UGSCore::FBuildStep& DefaultBuildStep : DefaultBuildSteps)
-		{
-			DefaultBuildStepObjects.Add(DefaultBuildStep.UniqueId, DefaultBuildStep.ToConfigObject());
-		}
-
-		return DefaultBuildStepObjects;
-	}
-
-	UGSCore::EBuildConfig GetEditorBuildConfig()
-	{
-		return UGSCore::EBuildConfig::Development;
-		//return ShouldSyncPrecompiledEditor() ? UGSCore::EBuildConfig::Development : Settings->CompiledEditorBuildConfig;
-	}
-
 	FString GetEditorExePath(UGSCore::EBuildConfig Config, TSharedPtr<UGSCore::FDetectProjectSettingsTask> DetectSettings)
 	{
 	#if PLATFORM_WINDOWS
@@ -149,24 +106,6 @@ namespace
 		return DetectSettings->BranchDirectoryName / "Engine" / "Binaries" / HostPlatform / ExeFileName;
 	}
 
-	TMap<FString, FString> GetWorkspaceVariables(TSharedPtr<UGSCore::FDetectProjectSettingsTask> DetectSettings)
-	{
-		UGSCore::EBuildConfig EditorBuildConfig = GetEditorBuildConfig();
-
-		TMap<FString, FString> Variables;
-		Variables.Add("BranchDir", DetectSettings->BranchDirectoryName);
-		Variables.Add("ProjectDir", FPaths::GetPath(DetectSettings->NewSelectedFileName));
-		Variables.Add("ProjectFile", DetectSettings->NewSelectedFileName);
-
-		// Todo: These might not be called "UE4*" anymore
-		Variables.Add("UE4EditorExe", GetEditorExePath(EditorBuildConfig, DetectSettings));
-		Variables.Add("UE4EditorCmdExe", GetEditorExePath(EditorBuildConfig, DetectSettings).Replace(TEXT(".exe"), TEXT("-Cmd.exe")));
-		Variables.Add("UE4EditorConfig", UGSCore::ToString(EditorBuildConfig));
-		Variables.Add("UE4EditorDebugArg", (EditorBuildConfig == UGSCore::EBuildConfig::Debug || EditorBuildConfig == UGSCore::EBuildConfig::DebugGame)? " -debug" : "");
-
-		return Variables;
-	}
-
 	static FString FormatUserName(FString UserName)
 	{
 		TStringBuilder<256> NormalUserName;
@@ -232,10 +171,66 @@ void UGSTab::CancelSync()
 	}
 }
 
+TMap<FString, FString> UGSTab::GetWorkspaceVariables() const
+{
+	UGSCore::EBuildConfig EditorBuildConfig = GetEditorBuildConfig();
+
+	TMap<FString, FString> Variables;
+	Variables.Add("BranchDir", DetectSettings->BranchDirectoryName);
+	Variables.Add("ProjectDir", FPaths::GetPath(DetectSettings->NewSelectedFileName));
+	Variables.Add("ProjectFile", DetectSettings->NewSelectedFileName);
+
+	// Todo: These might not be called "UE4*" anymore
+	Variables.Add("UE4EditorExe", GetEditorExePath(EditorBuildConfig, DetectSettings));
+	Variables.Add("UE4EditorCmdExe", GetEditorExePath(EditorBuildConfig, DetectSettings).Replace(TEXT(".exe"), TEXT("-Cmd.exe")));
+	Variables.Add("UE4EditorConfig", UGSCore::ToString(EditorBuildConfig));
+	Variables.Add("UE4EditorDebugArg", (EditorBuildConfig == UGSCore::EBuildConfig::Debug || EditorBuildConfig == UGSCore::EBuildConfig::DebugGame)? " -debug" : "");
+
+	return Variables;
+}
+
+UGSCore::EBuildConfig UGSTab::GetEditorBuildConfig() const
+{
+	return ShouldSyncPrecompiledEditor() ? UGSCore::EBuildConfig::Development : UserSettings->CompiledEditorBuildConfig;
+}
+
+bool UGSTab::ShouldSyncPrecompiledEditor() const
+{
+	return UserSettings->bSyncPrecompiledEditor && PerforceMonitor->HasZippedBinaries();
+}
+
+// Honestly ... seems ... super hacky/hardcoded. With out all these you assert when trying to merge build targets sooo a bit odd
+// TODO Need to do each of these ... per ... platform??
+TMap<FGuid, UGSCore::FCustomConfigObject> UGSTab::GetDefaultBuildStepObjects(const FString& EditorTargetName)
+{
+	TArray<UGSCore::FBuildStep> DefaultBuildSteps;
+	DefaultBuildSteps.Add(UGSCore::FBuildStep(FGuid(0x01F66060, 0x73FA4CC8, 0x9CB3E217, 0xFBBA954E), 0, TEXT("Compile UnrealHeaderTool"), TEXT("Compiling UnrealHeaderTool..."), 1, TEXT("UnrealHeaderTool"), HostPlatform, TEXT("Development"), TEXT(""), !ShouldSyncPrecompiledEditor()));
+	FString ActualEditorTargetName = (EditorTargetName.Len() > 0) ? EditorTargetName : "UnrealEditor";
+	DefaultBuildSteps.Add(UGSCore::FBuildStep(FGuid(0xF097FF61, 0xC9164058, 0x839135B4, 0x6C3173D5), 1, FString::Printf(TEXT("Compile %s"), *ActualEditorTargetName), FString::Printf(TEXT("Compiling %s..."), *ActualEditorTargetName), 10, ActualEditorTargetName, HostPlatform, UGSCore::ToString(UserSettings->CompiledEditorBuildConfig), TEXT(""), !ShouldSyncPrecompiledEditor()));
+	DefaultBuildSteps.Add(UGSCore::FBuildStep(FGuid(0xC6E633A1, 0x956F4AD3, 0xBC956D06, 0xD131E7B4), 2, TEXT("Compile ShaderCompileWorker"), TEXT("Compiling ShaderCompileWorker..."), 1, TEXT("ShaderCompileWorker"), HostPlatform, TEXT("Development"), TEXT(""), !ShouldSyncPrecompiledEditor()));
+	DefaultBuildSteps.Add(UGSCore::FBuildStep(FGuid(0x24FFD88C, 0x79014899, 0x9696AE10, 0x66B4B6E8), 3, TEXT("Compile UnrealLightmass"), TEXT("Compiling UnrealLightmass..."), 1, TEXT("UnrealLightmass"), HostPlatform, TEXT("Development"), TEXT(""), !ShouldSyncPrecompiledEditor()));
+	DefaultBuildSteps.Add(UGSCore::FBuildStep(FGuid(0xFFF20379, 0x06BF4205, 0x8A3EC534, 0x27736688), 4, TEXT("Compile CrashReportClient"), TEXT("Compiling CrashReportClient..."), 1, TEXT("CrashReportClient"), HostPlatform, TEXT("Development"), TEXT(""), !ShouldSyncPrecompiledEditor()));
+	DefaultBuildSteps.Add(UGSCore::FBuildStep(FGuid(0x89FE8A79, 0xD2594C7B, 0xBFB468F7, 0x218B91C2), 5, TEXT("Compile UnrealInsights"), TEXT("Compiling UnrealInsights..."), 1, TEXT("UnrealInsights"), HostPlatform, TEXT("Development"), TEXT(""), !ShouldSyncPrecompiledEditor()));
+	DefaultBuildSteps.Add(UGSCore::FBuildStep(FGuid(0x46312669, 0x5069428D, 0x8D72C241, 0x6C5A322E), 6, TEXT("Launch UnrealInsights"), TEXT("Running UnrealInsights..."), 1, TEXT("UnrealInsights"), HostPlatform, TEXT("Shipping"), TEXT(""), !ShouldSyncPrecompiledEditor()));
+	DefaultBuildSteps.Add(UGSCore::FBuildStep(FGuid(0xBB48CA5B, 0x56824432, 0x824DC451, 0x336A6523), 7, TEXT("Compile Zen Dashboard"), TEXT("Compile ZenDashboard Step..."), 1, TEXT("ZenDashboard"), HostPlatform, TEXT("Development"), TEXT(""), !ShouldSyncPrecompiledEditor()));
+	DefaultBuildSteps.Add(UGSCore::FBuildStep(FGuid(0x586CC0D3, 0x39144DF9, 0xACB62C02, 0xCD9D4FC6), 8, TEXT("Launch Zen Dashboard"), TEXT("Running Zen Dashboard..."), 1, TEXT("ZenDashboard"), HostPlatform, TEXT("Development"), TEXT(""), !ShouldSyncPrecompiledEditor()));
+	DefaultBuildSteps.Add(UGSCore::FBuildStep(FGuid(0x91C2A429, 0xC39149B4, 0x92A54E6B, 0xE71E0F00), 9, TEXT("Compile SwitchboardListener"), TEXT("Compiling SwitchboardListener..."), 1, TEXT("SwitchboardListener"), HostPlatform, TEXT("Development"), TEXT(""), !ShouldSyncPrecompiledEditor()));
+	DefaultBuildSteps.Add(UGSCore::FBuildStep(FGuid(0x5036C75B, 0x8DF04329, 0x82A1869D, 0xD2D48605), 10, TEXT("Compile UnrealMultiUserServer"), TEXT("Compiling UnrealMultiUserServer..."), 1, TEXT("UnrealMultiUserServer"), HostPlatform, TEXT("Development"), TEXT(""), !ShouldSyncPrecompiledEditor()));
+	DefaultBuildSteps.Add(UGSCore::FBuildStep(FGuid(0x274B89C3, 0x9DC64465, 0xA50840AB, 0xC4593CC2), 11, TEXT("Compile UnrealMultiUserSlateServer"), TEXT("Compiling UnrealMultiUserSlateServer..."), 1, TEXT("UnrealMultiUserSlateServer"), HostPlatform, TEXT("Development"), TEXT(""), !ShouldSyncPrecompiledEditor()));
+
+	TMap<FGuid, UGSCore::FCustomConfigObject> DefaultBuildStepObjects;
+	for (const UGSCore::FBuildStep& DefaultBuildStep : DefaultBuildSteps)
+	{
+		DefaultBuildStepObjects.Add(DefaultBuildStep.UniqueId, DefaultBuildStep.ToConfigObject());
+	}
+
+	return DefaultBuildStepObjects;
+}
+
 void UGSTab::OnSyncChangelist(int Changelist)
 {
 	// Options on what to do with workspace when updating it
-	UGSCore::EWorkspaceUpdateOptions Options = UGSCore::EWorkspaceUpdateOptions::Sync | UGSCore::EWorkspaceUpdateOptions::SyncArchives | UGSCore::EWorkspaceUpdateOptions::GenerateProjectFiles;
+	UGSCore::EWorkspaceUpdateOptions Options = UGSCore::EWorkspaceUpdateOptions::Sync | UGSCore::EWorkspaceUpdateOptions::GenerateProjectFiles;
 	if (UserSettings->bAutoResolveConflicts)
 	{
 		Options |= UGSCore::EWorkspaceUpdateOptions::AutoResolveChanges;
@@ -261,10 +256,20 @@ void UGSTab::OnSyncChangelist(int Changelist)
 		Changelist,
 		Options,
 		CombinedSyncFilter,
-		GetDefaultBuildStepObjects(DetectSettings->NewProjectEditorTarget, UserSettings),
+		GetDefaultBuildStepObjects(DetectSettings->NewProjectEditorTarget),
 		ProjectSettings->BuildSteps,
 		TSet<FGuid>(),
-		GetWorkspaceVariables(DetectSettings));
+		GetWorkspaceVariables());
+
+	if (UserSettings->bSyncPrecompiledEditor)
+	{
+		FString ArchivePath;
+		if (PerforceMonitor->TryGetArchivePathForChangeNumber(Changelist, ArchivePath))
+		{
+			Context->Options |= UGSCore::EWorkspaceUpdateOptions::SyncArchives;
+			Context->ArchiveTypeToDepotPath.Add(TPair<FString, FString>("Editor", ArchivePath));
+		}
+	}
 
 	// Update the workspace with the Context!
 	Workspace->Update(Context);
@@ -278,10 +283,10 @@ void UGSTab::OnBuildWorkspace()
 		-1,
 		Options,
 		CombinedSyncFilter,
-		GetDefaultBuildStepObjects(DetectSettings->NewProjectEditorTarget, UserSettings),
+		GetDefaultBuildStepObjects(DetectSettings->NewProjectEditorTarget),
 		ProjectSettings->BuildSteps,
 		TSet<FGuid>(),
-		GetWorkspaceVariables(DetectSettings));
+		GetWorkspaceVariables());
 
 	// Update the workspace with the Context!
 	Workspace->Update(Context);
@@ -691,7 +696,7 @@ bool UGSTab::SetupWorkspace()
 	PerforceMonitor = MakeShared<UGSCore::FPerforceMonitor>(PerforceClient.ToSharedRef(), BranchClientPath, SelectedClientFileName, SelectedProjectIdentifier, ProjectLogBaseName + ".p4.log");
 	PerforceMonitor->OnUpdate = [this]{ QueueMessageForMainThread([this] { UpdateGameTabBuildList(); }); };
 
-	//PerforceMonitor->OnUpdateMetadata = [this]{ printf("PerforceMonitor->OnUpdateMetadata\n"); }; //MessageQueue.Add([this]{ UpdateBuildMetadata(); }); };
+	//PerforceMonitor->OnUpdateMetadata = [this]{ printf("PerforceMonitor->OnUpdateMetadata %i\n", PerforceMonitor->HasZippedBinaries()); }; //MessageQueue.Add([this]{ UpdateBuildMetadata(); }); };
 	PerforceMonitor->OnStreamChange = [this]{ printf("PerforceMonitor->OnStreamChange\n"); }; // MessageQueue.Add([this]{ Owner->StreamChanged(this); }); };
 
 	/* TODO figure out if this is even working, and if so how to correctly use this
diff --git a/Engine/Source/Programs/SlateUGS/Private/UGSTab.h b/Engine/Source/Programs/SlateUGS/Private/UGSTab.h
--- a/Engine/Source/Programs/SlateUGS/Private/UGSTab.h
+++ b/Engine/Source/Programs/SlateUGS/Private/UGSTab.h
@@ -7,12 +7,12 @@
 #include "Widgets/SWorkspaceWindow.h"
 
 #include "UGSCore/Workspace.h"
+#include "UGSCore/UserSettings.h"
 
 namespace UGSCore
 {
 	struct FUserWorkspaceSettings;
 	struct FUserProjectSettings;
-	struct FUserSettings;
 	class FDetectProjectSettingsTask;
 	class FPerforceMonitor;
 	class FEventMonitor;
@@ -81,6 +81,11 @@ private:
 		UGSCore::EWorkspaceUpdateResult SyncResult,
 		const FString& StatusMessage);
 
+	TMap<FString, FString> GetWorkspaceVariables() const;
+	UGSCore::EBuildConfig GetEditorBuildConfig() const;
+	bool ShouldSyncPrecompiledEditor() const;
+	TMap<FGuid, UGSCore::FCustomConfigObject> GetDefaultBuildStepObjects(const FString& EditorTargetName);
+
 	// Allows the queuing of functions from threads to be run on the main thread
 	void QueueMessageForMainThread(TFunction<void()> Function);
 
