- Changes to get SlateUGS compiling without PCH or non-unity

diff --git a/Engine/Source/Programs/SlateUGS/Private/ScheduledSyncTimer.cpp b/Engine/Source/Programs/SlateUGS/Private/ScheduledSyncTimer.cpp
--- a/Engine/Source/Programs/SlateUGS/Private/ScheduledSyncTimer.cpp
+++ b/Engine/Source/Programs/SlateUGS/Private/ScheduledSyncTimer.cpp
@@ -2,6 +2,7 @@
 
 #include "ScheduledSyncTimer.h"
 #include "HAL/RunnableThread.h"
+#include "HAL/PlatformProcess.h"
 #include "UGSLog.h"
 
 ScheduledSyncTimer::~ScheduledSyncTimer()
diff --git a/Engine/Source/Programs/SlateUGS/Private/ScheduledSyncTimer.h b/Engine/Source/Programs/SlateUGS/Private/ScheduledSyncTimer.h
--- a/Engine/Source/Programs/SlateUGS/Private/ScheduledSyncTimer.h
+++ b/Engine/Source/Programs/SlateUGS/Private/ScheduledSyncTimer.h
@@ -4,6 +4,8 @@
 
 #include "HAL/Runnable.h"
 #include "Templates/Function.h"
+#include "Misc/DateTime.h"
+
 
 class ScheduledSyncTimer : FRunnable
 {
diff --git a/Engine/Source/Programs/SlateUGS/Private/SlateUGSApp.cpp b/Engine/Source/Programs/SlateUGS/Private/SlateUGSApp.cpp
--- a/Engine/Source/Programs/SlateUGS/Private/SlateUGSApp.cpp
+++ b/Engine/Source/Programs/SlateUGS/Private/SlateUGSApp.cpp
@@ -4,6 +4,7 @@
 #include "UGSLog.h"
 #include "RequiredProgramMainCPPInclude.h"
 #include "Framework/Application/SlateApplication.h"
+#include "Framework/Docking/TabManager.h"
 
 #include "Widgets/SGameSyncTab.h"
 #include "Widgets/SWorkspaceWindow.h"
diff --git a/Engine/Source/Programs/SlateUGS/Private/SlateUGSStyle.h b/Engine/Source/Programs/SlateUGS/Private/SlateUGSStyle.h
--- a/Engine/Source/Programs/SlateUGS/Private/SlateUGSStyle.h
+++ b/Engine/Source/Programs/SlateUGS/Private/SlateUGSStyle.h
@@ -3,6 +3,7 @@
 #pragma once
 
 #include "CoreMinimal.h"
+#include "Styling/SlateStyle.h"
 
 class FSlateUGSStyle : public FSlateStyleSet
 {
diff --git a/Engine/Source/Programs/SlateUGS/Private/UGSTab.cpp b/Engine/Source/Programs/SlateUGS/Private/UGSTab.cpp
--- a/Engine/Source/Programs/SlateUGS/Private/UGSTab.cpp
+++ b/Engine/Source/Programs/SlateUGS/Private/UGSTab.cpp
@@ -16,6 +16,7 @@
 #include "Widgets/SLogWidget.h"
 
 #include "HAL/PlatformApplicationMisc.h"
+#include "Framework/Application/SlateApplication.h"
 
 #include "Internationalization/Regex.h"
 
diff --git a/Engine/Source/Programs/SlateUGS/Private/UGSTab.h b/Engine/Source/Programs/SlateUGS/Private/UGSTab.h
--- a/Engine/Source/Programs/SlateUGS/Private/UGSTab.h
+++ b/Engine/Source/Programs/SlateUGS/Private/UGSTab.h
@@ -5,6 +5,7 @@
 #include "Widgets/SEmptyTab.h"
 #include "Widgets/SGameSyncTab.h"
 #include "Widgets/SWorkspaceWindow.h"
+#include "Widgets/Docking/SDockTab.h"
 
 #include "UGSCore/Workspace.h"
 #include "UGSCore/UserSettings.h"
diff --git a/Engine/Source/Programs/SlateUGS/Private/UGSTabManager.cpp b/Engine/Source/Programs/SlateUGS/Private/UGSTabManager.cpp
--- a/Engine/Source/Programs/SlateUGS/Private/UGSTabManager.cpp
+++ b/Engine/Source/Programs/SlateUGS/Private/UGSTabManager.cpp
@@ -3,6 +3,7 @@
 #include "UGSTabManager.h"
 #include "Containers/Array.h"
 #include "UGSLog.h"
+#include "HAL/FileManager.h"
 
 #include "UGSCore/UserSettings.h"
 
diff --git a/Engine/Source/Programs/SlateUGS/Private/UGSTabManager.h b/Engine/Source/Programs/SlateUGS/Private/UGSTabManager.h
--- a/Engine/Source/Programs/SlateUGS/Private/UGSTabManager.h
+++ b/Engine/Source/Programs/SlateUGS/Private/UGSTabManager.h
@@ -4,6 +4,7 @@
 
 #include "UGSTab.h"
 #include "ScheduledSyncTimer.h"
+#include "Widgets/Docking/SDockTab.h"
 
 namespace UGSCore
 {
diff --git a/Engine/Source/Programs/SlateUGS/Private/Widgets/SEmptyTab.cpp b/Engine/Source/Programs/SlateUGS/Private/Widgets/SEmptyTab.cpp
--- a/Engine/Source/Programs/SlateUGS/Private/Widgets/SEmptyTab.cpp
+++ b/Engine/Source/Programs/SlateUGS/Private/Widgets/SEmptyTab.cpp
@@ -2,12 +2,15 @@
 
 #include "SEmptyTab.h"
 
+#include "Framework/Application/SlateApplication.h"
 #include "SWorkspaceWindow.h"
 #include "SPrimaryButton.h"
 #include "SlateUGSStyle.h"
 
 #include "UGSTab.h"
 
+#include "Widgets/Images/SImage.h"
+
 // Todo: delete
 #include "Widgets/Colors/SSimpleGradient.h"
 
diff --git a/Engine/Source/Programs/SlateUGS/Private/Widgets/SGameSyncTab.cpp b/Engine/Source/Programs/SlateUGS/Private/Widgets/SGameSyncTab.cpp
--- a/Engine/Source/Programs/SlateUGS/Private/Widgets/SGameSyncTab.cpp
+++ b/Engine/Source/Programs/SlateUGS/Private/Widgets/SGameSyncTab.cpp
@@ -18,7 +18,9 @@
 #include "Widgets/Docking/SDockTab.h"
 #include "Widgets/Colors/SSimpleGradient.h"
 #include "Widgets/Images/SThrobber.h"
+#include "Widgets/Images/SImage.h"
 #include "Widgets/Layout/SSeparator.h"
+#include "Widgets/Views/SListView.h"
 
 #include "SlateUGSStyle.h"
 
diff --git a/Engine/Source/Programs/SlateUGS/Private/Widgets/SGameSyncTab.h b/Engine/Source/Programs/SlateUGS/Private/Widgets/SGameSyncTab.h
--- a/Engine/Source/Programs/SlateUGS/Private/Widgets/SGameSyncTab.h
+++ b/Engine/Source/Programs/SlateUGS/Private/Widgets/SGameSyncTab.h
@@ -3,6 +3,8 @@
 #pragma once
 
 #include "Widgets/SCompoundWidget.h"
+#include "Widgets/Views/STableRow.h"
+#include "Widgets/Text/STextBlock.h"
 #include "ChangeInfo.h"
 
 class UGSTab;
diff --git a/Engine/Source/Programs/SlateUGS/Private/Widgets/SHordeBadge.h b/Engine/Source/Programs/SlateUGS/Private/Widgets/SHordeBadge.h
--- a/Engine/Source/Programs/SlateUGS/Private/Widgets/SHordeBadge.h
+++ b/Engine/Source/Programs/SlateUGS/Private/Widgets/SHordeBadge.h
@@ -3,6 +3,7 @@
 #pragma once
 
 #include "Widgets/SCompoundWidget.h"
+#include "Widgets/Input/SButton.h"
 
 enum EBadgeState
 {
diff --git a/Engine/Source/Programs/SlateUGS/Private/Widgets/SLogWidget.cpp b/Engine/Source/Programs/SlateUGS/Private/Widgets/SLogWidget.cpp
--- a/Engine/Source/Programs/SlateUGS/Private/Widgets/SLogWidget.cpp
+++ b/Engine/Source/Programs/SlateUGS/Private/Widgets/SLogWidget.cpp
@@ -2,6 +2,7 @@
 
 #include "SLogWidget.h"
 #include "Framework/Text/SlateTextRun.h"
+#include "SlateOptMacros.h"
 
 BEGIN_SLATE_FUNCTION_BUILD_OPTIMIZATION
 
diff --git a/Engine/Source/Programs/SlateUGS/Private/Widgets/SWorkspaceWindow.cpp b/Engine/Source/Programs/SlateUGS/Private/Widgets/SWorkspaceWindow.cpp
--- a/Engine/Source/Programs/SlateUGS/Private/Widgets/SWorkspaceWindow.cpp
+++ b/Engine/Source/Programs/SlateUGS/Private/Widgets/SWorkspaceWindow.cpp
@@ -4,11 +4,14 @@
 
 #include "Styling/AppStyle.h"
 #include "DesktopPlatformModule.h"
+#include "Framework/Application/SlateApplication.h"
 
 #include "UGSTab.h"
 #include "SGameSyncTab.h"
 #include "SPopupTextWindow.h"
 #include "SPrimaryButton.h"
+#include "Widgets/Input/SCheckBox.h"
+#include "Widgets/Layout/SSpacer.h"
 
 #define LOCTEXT_NAMESPACE "UGSWorkspaceWindow"
 
diff --git a/Engine/Source/Programs/SlateUGS/Private/Widgets/SWorkspaceWindow.h b/Engine/Source/Programs/SlateUGS/Private/Widgets/SWorkspaceWindow.h
--- a/Engine/Source/Programs/SlateUGS/Private/Widgets/SWorkspaceWindow.h
+++ b/Engine/Source/Programs/SlateUGS/Private/Widgets/SWorkspaceWindow.h
@@ -3,6 +3,7 @@
 #pragma once
 
 #include "Widgets/SWindow.h"
+#include "Widgets/Input/SEditableTextBox.h"
 
 class UGSTab;
 
