Fix SlateUGS non-unity
#rnx
#horde 210043

diff --git a/Engine/Source/Programs/SlateUGS/Private/Widgets/SGameSyncTab.h b/Engine/Source/Programs/SlateUGS/Private/Widgets/SGameSyncTab.h
--- a/Engine/Source/Programs/SlateUGS/Private/Widgets/SGameSyncTab.h
+++ b/Engine/Source/Programs/SlateUGS/Private/Widgets/SGameSyncTab.h
@@ -3,6 +3,7 @@
 #pragma once
 
 #include "Widgets/SCompoundWidget.h"
+#include "Widgets/Views/SListView.h"
 #include "Widgets/Views/STableRow.h"
 #include "Widgets/Text/STextBlock.h"
 #include "ChangeInfo.h"
diff --git a/Engine/Source/Programs/SlateUGS/Private/Widgets/SPopupTextWindow.h b/Engine/Source/Programs/SlateUGS/Private/Widgets/SPopupTextWindow.h
--- a/Engine/Source/Programs/SlateUGS/Private/Widgets/SPopupTextWindow.h
+++ b/Engine/Source/Programs/SlateUGS/Private/Widgets/SPopupTextWindow.h
@@ -2,6 +2,7 @@
 
 #pragma once
 
+#include "Framework/Text/TextLayout.h"
 #include "Widgets/SWindow.h"
 
 class SPopupTextWindow : public SWindow
diff --git a/Engine/Source/Programs/SlateUGS/Private/Widgets/SScheduledSyncWindow.cpp b/Engine/Source/Programs/SlateUGS/Private/Widgets/SScheduledSyncWindow.cpp
--- a/Engine/Source/Programs/SlateUGS/Private/Widgets/SScheduledSyncWindow.cpp
+++ b/Engine/Source/Programs/SlateUGS/Private/Widgets/SScheduledSyncWindow.cpp
@@ -6,6 +6,7 @@
 #include "UGSTab.h"
 #include "UGSTabManager.h"
 
+#include "Framework/Application/SlateApplication.h"
 #include "Widgets/Input/SCheckBox.h"
 #include "Widgets/Input/SEditableTextBox.h"
 #include "Widgets/Layout/SUniformGridPanel.h"
diff --git a/Engine/Source/Programs/SlateUGS/Private/Widgets/SSettingsWindow.cpp b/Engine/Source/Programs/SlateUGS/Private/Widgets/SSettingsWindow.cpp
--- a/Engine/Source/Programs/SlateUGS/Private/Widgets/SSettingsWindow.cpp
+++ b/Engine/Source/Programs/SlateUGS/Private/Widgets/SSettingsWindow.cpp
@@ -5,6 +5,7 @@
 
 #include "UGSCore/UserSettings.h"
 
+#include "Framework/Application/SlateApplication.h"
 #include "Widgets/Input/SCheckBox.h"
 #include "Widgets/Layout/SHeader.h"
 #include "Widgets/Layout/SUniformGridPanel.h"
diff --git a/Engine/Source/Programs/SlateUGS/Private/Widgets/SSyncFilterWindow.cpp b/Engine/Source/Programs/SlateUGS/Private/Widgets/SSyncFilterWindow.cpp
--- a/Engine/Source/Programs/SlateUGS/Private/Widgets/SSyncFilterWindow.cpp
+++ b/Engine/Source/Programs/SlateUGS/Private/Widgets/SSyncFilterWindow.cpp
@@ -4,6 +4,7 @@
 
 #include "UGSTab.h"
 
+#include "Framework/Application/SlateApplication.h"
 #include "Widgets/Layout/SUniformGridPanel.h"
 #include "Widgets/Layout/SHeader.h"
 #include "SPrimaryButton.h"
