Add UI support for Type column and Code/Content badges in the build list

diff --git a/Engine/Source/Programs/SlateUGS/Private/Widgets/SGameSyncTab.cpp b/Engine/Source/Programs/SlateUGS/Private/Widgets/SGameSyncTab.cpp
--- a/Engine/Source/Programs/SlateUGS/Private/Widgets/SGameSyncTab.cpp
+++ b/Engine/Source/Programs/SlateUGS/Private/Widgets/SGameSyncTab.cpp
@@ -14,8 +14,6 @@
 #include "Widgets/SScheduledSyncWindow.h"
 #include "Widgets/SLogWidget.h"
 #include "Widgets/Layout/SHeader.h"
-#include "Widgets/Testing/STestSuite.h"
-#include "Widgets/Docking/SDockTab.h"
 #include "Widgets/Colors/SSimpleGradient.h"
 #include "Widgets/Images/SThrobber.h"
 #include "Widgets/Images/SImage.h"
@@ -32,6 +30,7 @@
 namespace
 {
 	const FName HordeTableColumnStatus(TEXT("Status"));
+	const FName HordeTableColumnType(TEXT("Type"));
 	const FName HordeTableColumnChange(TEXT("Change"));
 	const FName HordeTableColumnTime(TEXT("Time"));
 	const FName HordeTableColumnAuthor(TEXT("Author"));
@@ -47,31 +46,6 @@ void SBuildDataRow::Construct(const FArguments& InArgs, const TSharedRef<STableV
 
 TSharedRef<SWidget> SBuildDataRow::GenerateWidgetForColumn(const FName& ColumnId) // Todo: maybe can refactor some of this code so there's less duplication by using the root SWidget class on different types
 {
-	if (ColumnId == HordeTableColumnEditor)
-	{
-		// Todo: replace this dummy badge data with the real thing
-		return SNew(SHorizontalBox)
-			+SHorizontalBox::Slot()
-			[
-				SNew(SHordeBadge)
-				.Text(FText::FromString("Editor Win64"))
-				.BadgeState(static_cast<EBadgeState>(FMath::RandRange(0, 4)))
-			]
-			+SHorizontalBox::Slot()
-			.Padding(2.0f, 0.0f)
-			[
-				SNew(SHordeBadge)
-				.Text(FText::FromString("Editor Mac"))
-				.BadgeState(static_cast<EBadgeState>(FMath::RandRange(0, 4)))
-			]
-			+SHorizontalBox::Slot()
-			[
-				SNew(SHordeBadge)
-				.Text(FText::FromString("Editor Linux"))
-				.BadgeState(static_cast<EBadgeState>(FMath::RandRange(0, 4)))
-			];
-	}
-
 	if (ColumnId == HordeTableColumnStatus)
 	{
 		TSharedRef<SImage> StatusCircle = SNew(SImage).Image(FSlateUGSStyle::Get().GetBrush("Icons.FilledCircle"));
@@ -107,7 +81,42 @@ TSharedRef<SWidget> SBuildDataRow::GenerateWidgetForColumn(const FName& ColumnId
 				StatusCircle
 			];
 	}
+	
+	if (ColumnId == HordeTableColumnType)
+	{
+		// Todo: replace with real data
+		bool bIsCode = FMath::RandBool();
+		bool bIsContent = FMath::RandBool();
+
+		TSharedRef<SHorizontalBox> TypeBadges = SNew(SHorizontalBox);
+
+		if (bIsCode)
+		{
+			TypeBadges->AddSlot()
+			.Padding(1.0f, 0.0f)
+			.HAlign(HAlign_Center)
+			[
+				SNew(SHordeBadge)
+				.Text(FText::FromString("Code"))
+				.BadgeState(EBadgeState::Pending)
+			];
+		}
 
+		if (bIsContent)
+		{
+			TypeBadges->AddSlot()
+			.Padding(1.0f, 0.0f)
+			.HAlign(HAlign_Center)
+			[
+				SNew(SHordeBadge)
+				.Text(FText::FromString("Content"))
+				.BadgeState(EBadgeState::Pending)
+			];
+		}
+
+		return TypeBadges;
+	}
+	
 	TSharedRef<STextBlock> TextItem = SNew(STextBlock);
 	if (ColumnId == HordeTableColumnChange)
 	{
@@ -129,6 +138,30 @@ TSharedRef<SWidget> SBuildDataRow::GenerateWidgetForColumn(const FName& ColumnId
 			.Replace(TEXT("\n"), TEXT(" ")); // Todo: replacing both with spaces causes lots of double spaces, maybe filter that out?
 		TextItem->SetText(FText::FromString(Description));
 	}
+	if (ColumnId == HordeTableColumnEditor)
+	{
+		// Todo: replace this dummy badge data with the real thing
+		return SNew(SHorizontalBox)
+			+SHorizontalBox::Slot()
+			[
+				SNew(SHordeBadge)
+				.Text(FText::FromString("Editor Win64"))
+				.BadgeState(static_cast<EBadgeState>(FMath::RandRange(0, 4)))
+			]
+			+SHorizontalBox::Slot()
+			.Padding(2.0f, 0.0f)
+			[
+				SNew(SHordeBadge)
+				.Text(FText::FromString("Editor Mac"))
+				.BadgeState(static_cast<EBadgeState>(FMath::RandRange(0, 4)))
+			]
+			+SHorizontalBox::Slot()
+			[
+				SNew(SHordeBadge)
+				.Text(FText::FromString("Editor Linux"))
+				.BadgeState(static_cast<EBadgeState>(FMath::RandRange(0, 4)))
+			];
+	}
 
 	if (CurrentItem->bCurrentlySynced)
 	{
@@ -486,6 +519,9 @@ void SGameSyncTab::Construct(const FArguments& InArgs)
 					+SHeaderRow::Column(HordeTableColumnStatus)
 					.DefaultLabel(LOCTEXT("HordeHeaderStatus", ""))
 					.FixedWidth(35.0f)
+					+SHeaderRow::Column(HordeTableColumnType)
+					.DefaultLabel(LOCTEXT("HordeHeaderType", "Type"))
+					.FixedWidth(150.0f)
 					+SHeaderRow::Column(HordeTableColumnChange)
 					.DefaultLabel(LOCTEXT("HordeHeaderChange", "Change"))
 					.FixedWidth(75.0f)
@@ -499,7 +535,7 @@ void SGameSyncTab::Construct(const FArguments& InArgs)
 					.DefaultLabel(LOCTEXT("HordeHeaderDescription", "Description"))
 					+SHeaderRow::Column(HordeTableColumnEditor)
 					.DefaultLabel(LOCTEXT("HordeHeaderEditor", "Editor"))
-					.FillWidth(0.45f)
+					.FillWidth(0.5f)
 				)
 			]
 			// Log
