SlateUGS: Add new ChangeInfo for Precompiled binaries and Content/Code
Use the value for progress to display a percentage if there is one set
Change the text color when Sync Precompiled binaries are selected allowing you to see which CL have them

diff --git a/Engine/Source/Programs/SlateUGS/Private/ChangeInfo.h b/Engine/Source/Programs/SlateUGS/Private/ChangeInfo.h
--- a/Engine/Source/Programs/SlateUGS/Private/ChangeInfo.h
+++ b/Engine/Source/Programs/SlateUGS/Private/ChangeInfo.h
@@ -3,6 +3,7 @@
 #pragma once
 
 #include "UGSCore/EventMonitor.h"
+#include "UGSCore/PerforceMonitor.h"
 
 struct FChangeInfo
 {
@@ -12,6 +13,9 @@ struct FChangeInfo
 	UGSCore::EReviewVerdict ReviewStatus = UGSCore::EReviewVerdict::Unknown;
 	int Changelist = 0;
 	bool bCurrentlySynced = false;
+	bool bSyncingPrecompiled = false;
+	bool bHasZippedBinaries = false;
+	UGSCore::FChangeType ChangeType;
 	FText Author;
 	FString Description;
 
diff --git a/Engine/Source/Programs/SlateUGS/Private/UGSTab.cpp b/Engine/Source/Programs/SlateUGS/Private/UGSTab.cpp
--- a/Engine/Source/Programs/SlateUGS/Private/UGSTab.cpp
+++ b/Engine/Source/Programs/SlateUGS/Private/UGSTab.cpp
@@ -500,6 +500,16 @@ void UGSTab::UpdateGameTabBuildList()
 					}
 				}
 
+				bool bHasZippedBinaries = false;
+				FString Temp;
+				if (PerforceMonitor->TryGetArchivePathForChangeNumber(Change.Number, Temp))
+				{
+					bHasZippedBinaries = true;
+				}
+
+				UGSCore::FChangeType ChangeType;
+				PerforceMonitor->TryGetChangeType(Change.Number, ChangeType);
+
 				if (ChangeIdx == 0 ||
 					Changes[ChangeIdx - 1]->Date.GetDayOfYear() != Changes[ChangeIdx]->Date.GetDayOfYear())
 				{
@@ -520,6 +530,9 @@ void UGSTab::UpdateGameTabBuildList()
 					Status,
 					Change.Number,
 					bCurrentSyncedChange,
+					ShouldSyncPrecompiledEditor(),
+					bHasZippedBinaries,
+					ChangeType,
 					FText::FromString(FormatUserName(Change.User)),
 					Change.Description
 				}));
@@ -537,7 +550,14 @@ bool UGSTab::IsSyncing() const
 
 FString UGSTab::GetSyncProgress() const
 {
-	return Workspace->GetCurrentProgress().Key;
+	TTuple<FString, float> CurrentProgress = Workspace->GetCurrentProgress();
+
+	if (CurrentProgress.Value > 0.0f)
+	{
+		return FString::Printf(TEXT("%s %.2f%%"), *CurrentProgress.Key, CurrentProgress.Value * 100.0f);
+	}
+
+	return CurrentProgress.Key;
 }
 
 TArray<UGSCore::FWorkspaceSyncCategory> UGSTab::GetSyncCategories(SyncCategoryType CategoryType) const
@@ -697,7 +717,8 @@ bool UGSTab::SetupWorkspace()
 	PerforceMonitor = MakeShared<UGSCore::FPerforceMonitor>(PerforceClient.ToSharedRef(), BranchClientPath, SelectedClientFileName, SelectedProjectIdentifier, ProjectLogBaseName + ".p4.log");
 	PerforceMonitor->OnUpdate = [this]{ QueueMessageForMainThread([this] { UpdateGameTabBuildList(); }); };
 
-	//PerforceMonitor->OnUpdateMetadata = [this]{ printf("PerforceMonitor->OnUpdateMetadata %i\n", PerforceMonitor->HasZippedBinaries()); }; //MessageQueue.Add([this]{ UpdateBuildMetadata(); }); };
+	//PerforceMonitor->OnUpdateMetadata = [this]{ QueueMessageForMainThread([this] { UpdateGameTabBuildList(); }); }; //MessageQueue.Add([this]{ UpdateBuildMetadata(); }); };
+	PerforceMonitor->OnChangeTypeQueryFinished = [this]{ QueueMessageForMainThread([this] { UpdateGameTabBuildList(); }); };
 	PerforceMonitor->OnStreamChange = [this]{ printf("PerforceMonitor->OnStreamChange\n"); }; // MessageQueue.Add([this]{ Owner->StreamChanged(this); }); };
 
 	/* TODO figure out if this is even working, and if so how to correctly use this
diff --git a/Engine/Source/Programs/SlateUGS/Private/Widgets/SGameSyncTab.cpp b/Engine/Source/Programs/SlateUGS/Private/Widgets/SGameSyncTab.cpp
--- a/Engine/Source/Programs/SlateUGS/Private/Widgets/SGameSyncTab.cpp
+++ b/Engine/Source/Programs/SlateUGS/Private/Widgets/SGameSyncTab.cpp
@@ -81,12 +81,11 @@ TSharedRef<SWidget> SBuildDataRow::GenerateWidgetForColumn(const FName& ColumnId
 				StatusCircle
 			];
 	}
-	
+
 	if (ColumnId == HordeTableColumnType)
 	{
-		// Todo: replace with real data
-		bool bIsCode = FMath::RandBool();
-		bool bIsContent = FMath::RandBool();
+		bool bIsCode    = CurrentItem->ChangeType.bContainsCode;
+		bool bIsContent = CurrentItem->ChangeType.bContainsContent;
 
 		TSharedRef<SHorizontalBox> TypeBadges = SNew(SHorizontalBox);
 
@@ -116,7 +115,7 @@ TSharedRef<SWidget> SBuildDataRow::GenerateWidgetForColumn(const FName& ColumnId
 
 		return TypeBadges;
 	}
-	
+
 	TSharedRef<STextBlock> TextItem = SNew(STextBlock);
 	if (ColumnId == HordeTableColumnChange)
 	{
@@ -163,6 +162,14 @@ TSharedRef<SWidget> SBuildDataRow::GenerateWidgetForColumn(const FName& ColumnId
 			];
 	}
 
+	if (CurrentItem->bSyncingPrecompiled)
+	{
+		if (!CurrentItem->bHasZippedBinaries)
+		{
+			TextItem->SetColorAndOpacity(FLinearColor(0.1f, 0.1f, 0.1f));
+		}
+	}
+
 	if (CurrentItem->bCurrentlySynced)
 	{
 		// Lets make the font white and bold when we are the currently synced CL
diff --git a/Engine/Source/Programs/SlateUGS/Private/Widgets/SSettingsWindow.cpp b/Engine/Source/Programs/SlateUGS/Private/Widgets/SSettingsWindow.cpp
--- a/Engine/Source/Programs/SlateUGS/Private/Widgets/SSettingsWindow.cpp
+++ b/Engine/Source/Programs/SlateUGS/Private/Widgets/SSettingsWindow.cpp
@@ -245,6 +245,9 @@ void SSettingsWindow::HandleSyncCompiledEditor(ECheckBoxState InCheck)
 	{
 		UserSettings->bSyncPrecompiledEditor = bNewValue;
 		UserSettings->Save();
+
+		// If we update this setting, update the build list to refresh filtering
+		Tab->UpdateGameTabBuildList();
 	}
 }
 
diff --git a/Engine/Source/Programs/UGSCore/PerforceMonitor.cpp b/Engine/Source/Programs/UGSCore/PerforceMonitor.cpp
--- a/Engine/Source/Programs/UGSCore/PerforceMonitor.cpp
+++ b/Engine/Source/Programs/UGSCore/PerforceMonitor.cpp
@@ -422,6 +422,12 @@ bool FPerforceMonitor::UpdateChangeTypes()
 			OnUpdateMetadata();
 		}
 	}
+
+	if (OnChangeTypeQueryFinished)
+	{
+		OnChangeTypeQueryFinished();
+	}
+
 	return true;
 }
 
diff --git a/Engine/Source/Programs/UGSCore/PerforceMonitor.h b/Engine/Source/Programs/UGSCore/PerforceMonitor.h
--- a/Engine/Source/Programs/UGSCore/PerforceMonitor.h
+++ b/Engine/Source/Programs/UGSCore/PerforceMonitor.h
@@ -25,6 +25,7 @@ public:
 
 	TFunction<void()> OnUpdate;
 	TFunction<void()> OnUpdateMetadata;
+	TFunction<void()> OnChangeTypeQueryFinished;
 	TFunction<void()> OnStreamChange;
 
 	FPerforceMonitor(const TSharedRef<FPerforceConnection>& InPerforce, const FString& InBranchClientPath, const FString& InSelectedClientFileName, const FString& InSelectedProjectIdentifier, const FString& InLogPath);
