Add UI (without functionality) for new workspace creation window

diff --git a/Engine/Source/Programs/SlateUGS/Private/Widgets/SNewWorkspaceWindow.cpp b/Engine/Source/Programs/SlateUGS/Private/Widgets/SNewWorkspaceWindow.cpp
new file mode 100644
--- /dev/null
+++ b/Engine/Source/Programs/SlateUGS/Private/Widgets/SNewWorkspaceWindow.cpp
@@ -0,0 +1,155 @@
+// Copyright Epic Games, Inc. All Rights Reserved.
+
+#include "SNewWorkspaceWindow.h"
+
+#include "SlateUGSStyle.h"
+#include "Framework/Application/SlateApplication.h"
+
+#include "UGSTab.h"
+#include "SGameSyncTab.h"
+#include "SPrimaryButton.h"
+#include "Widgets/Layout/SHeader.h"
+
+#define LOCTEXT_NAMESPACE "UGSNewWorkspaceWindow"
+
+void SNewWorkspaceWindow::Construct(const FArguments& InArgs)
+{
+	Tab = InArgs._Tab;
+
+	SWindow::Construct(SWindow::FArguments()
+	.Title(LOCTEXT("WindowTitle", "New Workspace"))
+	.SizingRule(ESizingRule::FixedSize)
+	.ClientSize(FVector2D(800, 200))
+	[
+		SNew(SBox)
+		.Padding(30.0f, 15.0f, 30.0f, 0.0f)
+		[
+			SNew(SVerticalBox)
+			+SVerticalBox::Slot()
+			.AutoHeight()
+			[
+				SNew(SHeader)
+				[
+					SNew(STextBlock)
+					.Text(LOCTEXT("CustomView", "Settings"))
+				]
+			]
+			+SVerticalBox::Slot()
+			.AutoHeight()
+			.VAlign(VAlign_Bottom)
+			.Padding(40.0f, 20.0f)
+			[
+				SNew(SHorizontalBox)
+				+SHorizontalBox::Slot()
+				.FillWidth(0.25f)
+				[
+					SNew(SVerticalBox)
+					+SVerticalBox::Slot()
+					[
+						SNew(STextBlock)
+						.Text(LOCTEXT("StreamText", "Stream:"))
+					]
+					+SVerticalBox::Slot()
+					.Padding(0.0f, 7.5f)
+					[
+						SNew(STextBlock)
+						.Text(LOCTEXT("RootDirectoryText", "Root Directory:"))
+					]
+					+SVerticalBox::Slot()
+					[
+						SNew(STextBlock)
+						.Text(LOCTEXT("NameText", "Name:"))
+					]
+				]
+				+SHorizontalBox::Slot()
+				[
+					SNew(SVerticalBox)
+					+SVerticalBox::Slot()
+					[
+						SNew(SHorizontalBox)
+						+SHorizontalBox::Slot()
+						[
+							SNew(SEditableTextBox)
+						]
+						+SHorizontalBox::Slot()
+						.FillWidth(0.225f)
+						.HAlign(HAlign_Right)
+						[
+							SNew(SButton)
+							.Text(LOCTEXT("BrowseStreamButtonText", "Browse..."))
+							//.OnClicked(this, &SNewWorkspaceWindow::/* callback here */)
+						]
+					]
+					+SVerticalBox::Slot()
+					.Padding(0.0f, 7.5f)
+					[
+						SNew(SHorizontalBox)
+						+SHorizontalBox::Slot()
+						[
+							SNew(SEditableTextBox)
+						]
+						+SHorizontalBox::Slot()
+						.FillWidth(0.225f)
+						.HAlign(HAlign_Right)
+						[
+							SNew(SButton)
+							.Text(LOCTEXT("BrowseRootDirectoryButtonText", "Browse..."))
+							//.OnClicked(this, &SNewWorkspaceWindow::/* callback here */)
+						]
+					]
+					+SVerticalBox::Slot()
+					[
+						SNew(SEditableTextBox)
+					]
+				]
+			]
+			+SVerticalBox::Slot()
+			.AutoHeight()
+			.VAlign(VAlign_Bottom)
+			[
+				SNew(SBox)
+				.HAlign(HAlign_Right)
+				.Padding(10.0f)
+				[
+					SNew(SHorizontalBox)
+					+SHorizontalBox::Slot()
+					.Padding(0.0f, 0.0f, 10.0f, 0.0f)
+					[
+						SNew(SPrimaryButton)
+						.Text(LOCTEXT("CreateButtonText", "Create"))
+						.OnClicked(this, &SNewWorkspaceWindow::OnCreateClicked)
+					]
+					+SHorizontalBox::Slot()
+					[
+						SNew(SButton)
+						.Text(LOCTEXT("CancelButtonText", "Cancel"))
+						.OnClicked(this, &SNewWorkspaceWindow::OnCancelClicked)
+					]
+				]
+			]
+		]
+	]);
+}
+
+FReply SNewWorkspaceWindow::OnBrowseStreamClicked()
+{
+	return FReply::Handled();
+}
+
+FReply SNewWorkspaceWindow::OnBrowseRootDirectoryClicked()
+{
+	return FReply::Handled();
+}
+
+FReply SNewWorkspaceWindow::OnCreateClicked()
+{	
+	return FReply::Handled();
+}
+
+FReply SNewWorkspaceWindow::OnCancelClicked()
+{
+	FSlateApplication::Get().FindWidgetWindow(AsShared())->RequestDestroyWindow();
+	return FReply::Handled();
+}
+
+#undef LOCTEXT_NAMESPACE
diff --git a/Engine/Source/Programs/SlateUGS/Private/Widgets/SNewWorkspaceWindow.h b/Engine/Source/Programs/SlateUGS/Private/Widgets/SNewWorkspaceWindow.h
new file mode 100644
--- /dev/null
+++ b/Engine/Source/Programs/SlateUGS/Private/Widgets/SNewWorkspaceWindow.h
@@ -0,0 +1,30 @@
+// Copyright Epic Games, Inc. All Rights Reserved.
+
+#pragma once
+
+#include "Widgets/SWindow.h"
+#include "Widgets/Input/SEditableTextBox.h"
+
+class UGSTab;
+
+class SNewWorkspaceWindow final : public SWindow
+{
+public:
+	SLATE_BEGIN_ARGS(SNewWorkspaceWindow) {}
+	SLATE_ARGUMENT(UGSTab*, Tab)
+SLATE_END_ARGS()
+
+void Construct(const FArguments& InArgs);
+
+private:
+	FReply OnBrowseStreamClicked();
+	FReply OnBrowseRootDirectoryClicked();
+
+	FReply OnCreateClicked();
+	FReply OnCancelClicked();
+
+	TSharedPtr<SEditableTextBox> LocalFileText = nullptr;
+	FString WorkspacePathText;
+
+	UGSTab* Tab = nullptr;
+};
diff --git a/Engine/Source/Programs/SlateUGS/Private/Widgets/SWorkspaceWindow.cpp b/Engine/Source/Programs/SlateUGS/Private/Widgets/SWorkspaceWindow.cpp
--- a/Engine/Source/Programs/SlateUGS/Private/Widgets/SWorkspaceWindow.cpp
+++ b/Engine/Source/Programs/SlateUGS/Private/Widgets/SWorkspaceWindow.cpp
@@ -8,6 +8,7 @@
 
 #include "UGSTab.h"
 #include "SGameSyncTab.h"
+#include "SNewWorkspaceWindow.h"
 #include "SPopupTextWindow.h"
 #include "SPrimaryButton.h"
 #include "Widgets/Input/SCheckBox.h"
@@ -97,7 +98,6 @@ void SWorkspaceWindow::Construct(const FArguments& InArgs)
 					.AutoWidth()
 					[
 						SNew(SCheckBox).Style(FAppStyle::Get(), "RadioButton")
-						.IsEnabled(false) // Todo: enable after adding this functionality
 						.IsChecked_Lambda([this] () { return bIsLocalFileSelected ? ECheckBoxState::Unchecked : ECheckBoxState::Checked; })
 						.OnCheckStateChanged_Lambda( [this] (ECheckBoxState InState) { bIsLocalFileSelected = (!bIsLocalFileSelected && InState == ECheckBoxState::Checked); } )
 					]
@@ -136,6 +136,7 @@ void SWorkspaceWindow::Construct(const FArguments& InArgs)
 							SNew(SButton)
 							.HAlign(HAlign_Center)
 							.Text(LOCTEXT("NewText", "New..."))
+							.OnClicked(this, &SWorkspaceWindow::OnNewClicked)
 						]
 						+SHorizontalBox::Slot()
 						.Padding(10.0f, 0.0f, 0.0f, 0.0f)
@@ -263,4 +264,11 @@ FReply SWorkspaceWindow::OnBrowseClicked()
 	return FReply::Handled();
 }
 
+FReply SWorkspaceWindow::OnNewClicked()
+{
+	FSlateApplication::Get().AddModalWindow(SNew(SNewWorkspaceWindow).Tab(Tab), SharedThis(this), false);
+
+	return FReply::Handled();
+}
+
 #undef LOCTEXT_NAMESPACE
diff --git a/Engine/Source/Programs/SlateUGS/Private/Widgets/SWorkspaceWindow.h b/Engine/Source/Programs/SlateUGS/Private/Widgets/SWorkspaceWindow.h
--- a/Engine/Source/Programs/SlateUGS/Private/Widgets/SWorkspaceWindow.h
+++ b/Engine/Source/Programs/SlateUGS/Private/Widgets/SWorkspaceWindow.h
@@ -17,11 +17,12 @@ public:
 	void Construct(const FArguments& InArgs);
 
 private:
-	FReply OnBrowseClicked();
-
 	FReply OnOkClicked();
 	FReply OnCancelClicked();
 
+	FReply OnBrowseClicked();
+	FReply OnNewClicked();
+
 	bool bIsLocalFileSelected = true;
 	TSharedPtr<SEditableTextBox> LocalFileText = nullptr;
 	FString WorkspacePathText;
