Added UI for select-a-stream window, but data is not shown yet.

Made all instances of UGSTab* null-initialized, and made them required arguments in widgets which hold them

diff --git a/Engine/Source/Programs/SlateUGS/Private/UGSTab.cpp b/Engine/Source/Programs/SlateUGS/Private/UGSTab.cpp
--- a/Engine/Source/Programs/SlateUGS/Private/UGSTab.cpp
+++ b/Engine/Source/Programs/SlateUGS/Private/UGSTab.cpp
@@ -24,8 +24,8 @@
 
 UGSTab::UGSTab() : TabArgs(nullptr, FTabId()),
 				   TabWidget(SNew(SDockTab)),
-				   EmptyTabView(SNew(SEmptyTab).Tab(this)),
-				   GameSyncTabView(SNew(SGameSyncTab).Tab(this)),
+				   EmptyTabView(SNew(SEmptyTab, this)),
+				   GameSyncTabView(SNew(SGameSyncTab, this)),
 				   bHasQueuedMessages(false),
 				   bNeedUpdateGameTabBuildList(false)
 {
@@ -46,6 +46,8 @@ void UGSTab::Initialize(TSharedPtr<UGSCore::FUserSettings> InUserSettings)
 
 	TabWidget->SetContent(EmptyTabView);
 	TabWidget->SetLabel(FText(LOCTEXT("TabName", "Select a Project")));
+
+	PerforceClient = MakeShared<UGSCore::FPerforceConnection>(TEXT(""), TEXT(""), TEXT(""));
 }
 
 void UGSTab::SetTabManager(UGSTabManager* InTabManager)
@@ -200,6 +202,15 @@ bool UGSTab::ShouldSyncPrecompiledEditor() const
 	return UserSettings->bSyncPrecompiledEditor && PerforceMonitor->HasZippedBinaries();
 }
 
+TArray<FString> UGSTab::GetAllStreamNames() const
+{
+	TArray<FString> Result;
+	FEvent* AbortEvent = FPlatformProcess::GetSynchEventFromPool(true);
+	PerforceClient->FindStreams("//*/*", Result, AbortEvent, *MakeShared<FLineWriter>());
+
+	return Result;
+}
+
 // Honestly ... seems ... super hacky/hardcoded. With out all these you assert when trying to merge build targets sooo a bit odd
 // TODO Need to do each of these ... per ... platform??
 TMap<FGuid, UGSCore::FCustomConfigObject> UGSTab::GetDefaultBuildStepObjects(const FString& EditorTargetName)
@@ -644,7 +655,7 @@ bool UGSTab::SetupWorkspace()
 	ProjectFileName = UGSCore::FUtility::GetPathWithCorrectCase(ProjectFileName);
 
 	// TODO likely should also log this on an Empty tab... so we can show logging info when we are loading things
-	DetectSettings = MakeShared<UGSCore::FDetectProjectSettingsTask>(MakeShared<UGSCore::FPerforceConnection>(TEXT(""), TEXT(""), TEXT("")), ProjectFileName, MakeShared<FLineWriter>());
+	DetectSettings = MakeShared<UGSCore::FDetectProjectSettingsTask>(PerforceClient.ToSharedRef(), ProjectFileName, MakeShared<FLineWriter>());
 
 	TSharedRef<UGSCore::FModalTaskResult> Result = ExecuteModalTask(
 		FSlateApplication::Get().GetActiveModalWindow(),
diff --git a/Engine/Source/Programs/SlateUGS/Private/UGSTab.h b/Engine/Source/Programs/SlateUGS/Private/UGSTab.h
--- a/Engine/Source/Programs/SlateUGS/Private/UGSTab.h
+++ b/Engine/Source/Programs/SlateUGS/Private/UGSTab.h
@@ -73,6 +73,7 @@ public:
 	UGSTabManager* GetTabManager();
 	TSharedPtr<UGSCore::FUserSettings> GetUserSettings() const;
 	bool ShouldSyncPrecompiledEditor() const;
+	TArray<FString> GetAllStreamNames() const;
 
 	void UpdateGameTabBuildList();
 	void RefreshBuildList();
diff --git a/Engine/Source/Programs/SlateUGS/Private/Widgets/SEmptyTab.cpp b/Engine/Source/Programs/SlateUGS/Private/Widgets/SEmptyTab.cpp
--- a/Engine/Source/Programs/SlateUGS/Private/Widgets/SEmptyTab.cpp
+++ b/Engine/Source/Programs/SlateUGS/Private/Widgets/SEmptyTab.cpp
@@ -11,14 +11,11 @@
 
 #include "Widgets/Images/SImage.h"
 
-// Todo: delete
-#include "Widgets/Colors/SSimpleGradient.h"
-
 #define LOCTEXT_NAMESPACE "UGSEmptyTab"
 
-void SEmptyTab::Construct(const FArguments& InArgs)
+void SEmptyTab::Construct(const FArguments& InArgs, UGSTab* InTab)
 {
-	Tab = InArgs._Tab;
+	Tab = InTab;
 
 	this->ChildSlot
 	[
@@ -75,7 +72,7 @@ void SEmptyTab::Construct(const FArguments& InArgs)
 FReply SEmptyTab::OnOpenProjectClicked()
 {
 	FSlateApplication& SlateApplication = FSlateApplication::Get();
-	SlateApplication.AddModalWindow(SNew(SWorkspaceWindow).Tab(Tab), Tab->GetTabArgs().GetOwnerWindow(), false);
+	SlateApplication.AddModalWindow(SNew(SWorkspaceWindow, Tab), Tab->GetTabArgs().GetOwnerWindow(), false);
 
 	return FReply::Handled();
 }
diff --git a/Engine/Source/Programs/SlateUGS/Private/Widgets/SEmptyTab.h b/Engine/Source/Programs/SlateUGS/Private/Widgets/SEmptyTab.h
--- a/Engine/Source/Programs/SlateUGS/Private/Widgets/SEmptyTab.h
+++ b/Engine/Source/Programs/SlateUGS/Private/Widgets/SEmptyTab.h
@@ -6,17 +6,16 @@
 
 class UGSTab;
 
-class SEmptyTab : public SCompoundWidget
+class SEmptyTab final : public SCompoundWidget
 {
 public:
 	SLATE_BEGIN_ARGS(SEmptyTab) {}
-		SLATE_ARGUMENT(UGSTab*, Tab)
 	SLATE_END_ARGS()
 
-	void Construct(const FArguments& InArgs);
+	void Construct(const FArguments& InArgs, UGSTab* InTab);
 
 private:
 	FReply OnOpenProjectClicked();
 
-	UGSTab* Tab;
+	UGSTab* Tab = nullptr;
 };
diff --git a/Engine/Source/Programs/SlateUGS/Private/Widgets/SGameSyncTab.cpp b/Engine/Source/Programs/SlateUGS/Private/Widgets/SGameSyncTab.cpp
--- a/Engine/Source/Programs/SlateUGS/Private/Widgets/SGameSyncTab.cpp
+++ b/Engine/Source/Programs/SlateUGS/Private/Widgets/SGameSyncTab.cpp
@@ -237,7 +237,7 @@ TSharedRef<SWidget> SGameSyncTab::MakeSyncButtonDropdown()
 		FText::FromString(TEXT("Setup Schedule Sync")),
 		FText::FromString(TEXT("Setup a schedule sync to run at specific time")),
 		FSlateIcon(),
-		FUIAction(FExecuteAction::CreateLambda([this] { FSlateApplication::Get().AddModalWindow(SNew(SScheduledSyncWindow).Tab(Tab), Tab->GetTabArgs().GetOwnerWindow(), false); } )),
+		FUIAction(FExecuteAction::CreateLambda([this] { FSlateApplication::Get().AddModalWindow(SNew(SScheduledSyncWindow, Tab), Tab->GetTabArgs().GetOwnerWindow(), false); } )),
 		NAME_None,
 		EUserInterfaceActionType::Button
 	);
@@ -245,9 +245,9 @@ TSharedRef<SWidget> SGameSyncTab::MakeSyncButtonDropdown()
 	return MenuBuilder.MakeWidget();
 }
 
-void SGameSyncTab::Construct(const FArguments& InArgs)
+void SGameSyncTab::Construct(const FArguments& InArgs, UGSTab* InTab)
 {
-	Tab = InArgs._Tab;
+	Tab = InTab;
 
 	this->ChildSlot
 	[
@@ -374,7 +374,7 @@ void SGameSyncTab::Construct(const FArguments& InArgs)
 						.Text(LOCTEXT("Filter", "Filter"))
 						.Icon(FSlateUGSStyle::Get().GetBrush("Icons.Filter"))
 						// Todo: this is probably the wrong "Filter" button. The functionality below should probably be in the settings dropdown
-						.OnClicked_Lambda([this] { FSlateApplication::Get().AddModalWindow(SNew(SSyncFilterWindow).Tab(Tab), Tab->GetTabArgs().GetOwnerWindow(), false); return FReply::Handled(); })
+						.OnClicked_Lambda([this] { FSlateApplication::Get().AddModalWindow(SNew(SSyncFilterWindow, Tab), Tab->GetTabArgs().GetOwnerWindow(), false); return FReply::Handled(); })
 					]
 					+SHorizontalBox::Slot()
 					.AutoWidth()
@@ -383,7 +383,7 @@ void SGameSyncTab::Construct(const FArguments& InArgs)
 						SNew(SSimpleButton)
 						.Text(LOCTEXT("Settings", "Settings"))
 						.Icon(FSlateUGSStyle::Get().GetBrush("Icons.Settings"))
-						.OnClicked_Lambda([this] { FSlateApplication::Get().AddModalWindow(SNew(SSettingsWindow).Tab(Tab), Tab->GetTabArgs().GetOwnerWindow(), false); return FReply::Handled(); })
+						.OnClicked_Lambda([this] { FSlateApplication::Get().AddModalWindow(SNew(SSettingsWindow, Tab), Tab->GetTabArgs().GetOwnerWindow(), false); return FReply::Handled(); })
 					]
 				]
 			]
diff --git a/Engine/Source/Programs/SlateUGS/Private/Widgets/SGameSyncTab.h b/Engine/Source/Programs/SlateUGS/Private/Widgets/SGameSyncTab.h
--- a/Engine/Source/Programs/SlateUGS/Private/Widgets/SGameSyncTab.h
+++ b/Engine/Source/Programs/SlateUGS/Private/Widgets/SGameSyncTab.h
@@ -27,10 +27,9 @@ class SGameSyncTab final : public SCompoundWidget
 {
 public:
 	SLATE_BEGIN_ARGS(SGameSyncTab) {}
-		SLATE_ARGUMENT(UGSTab*, Tab)
 	SLATE_END_ARGS()
 
-	void Construct(const FArguments& InArgs);
+	void Construct(const FArguments& InArgs, UGSTab* InTab);
 
 	// We need access to the SyncLog when creating the Workspace
 	// TODO: think of a better way to do this
@@ -64,5 +63,5 @@ private:
 	static constexpr float HordeBuildRowVerticalPadding = 2.5f;
 	static constexpr float HordeBuildRowExtraIconPadding = 10.0f;
 
-	UGSTab* Tab;
+	UGSTab* Tab = nullptr;
 };
diff --git a/Engine/Source/Programs/SlateUGS/Private/Widgets/SNewWorkspaceWindow.cpp b/Engine/Source/Programs/SlateUGS/Private/Widgets/SNewWorkspaceWindow.cpp
--- a/Engine/Source/Programs/SlateUGS/Private/Widgets/SNewWorkspaceWindow.cpp
+++ b/Engine/Source/Programs/SlateUGS/Private/Widgets/SNewWorkspaceWindow.cpp
@@ -9,12 +9,13 @@
 #include "SGameSyncTab.h"
 #include "SPrimaryButton.h"
 #include "Widgets/Layout/SHeader.h"
+#include "Widgets/SSelectStreamWindow.h"
 
 #define LOCTEXT_NAMESPACE "UGSNewWorkspaceWindow"
 
-void SNewWorkspaceWindow::Construct(const FArguments& InArgs)
+void SNewWorkspaceWindow::Construct(const FArguments& InArgs, UGSTab* InTab)
 {
-	Tab = InArgs._Tab;
+	Tab = InTab;
 
 	SWindow::Construct(SWindow::FArguments()
 	.Title(LOCTEXT("WindowTitle", "New Workspace"))
@@ -77,7 +78,7 @@ void SNewWorkspaceWindow::Construct(const FArguments& InArgs)
 						[
 							SNew(SButton)
 							.Text(LOCTEXT("BrowseStreamButtonText", "Browse..."))
-							//.OnClicked(this, &SNewWorkspaceWindow::/* callback here */)
+							.OnClicked(this, &SNewWorkspaceWindow::OnBrowseStreamClicked)
 						]
 					]
 					+SVerticalBox::Slot()
@@ -94,7 +95,7 @@ void SNewWorkspaceWindow::Construct(const FArguments& InArgs)
 						[
 							SNew(SButton)
 							.Text(LOCTEXT("BrowseRootDirectoryButtonText", "Browse..."))
-							//.OnClicked(this, &SNewWorkspaceWindow::/* callback here */)
+							.OnClicked(this, &SNewWorkspaceWindow::OnBrowseRootDirectoryClicked)
 						]
 					]
 					+SVerticalBox::Slot()
@@ -133,6 +134,8 @@ void SNewWorkspaceWindow::Construct(const FArguments& InArgs)
 
 FReply SNewWorkspaceWindow::OnBrowseStreamClicked()
 {
+	FSlateApplication::Get().AddModalWindow(SNew(SSelectStreamWindow, Tab), SharedThis(this), false);
+
 	return FReply::Handled();
 }
 
@@ -142,7 +145,7 @@ FReply SNewWorkspaceWindow::OnBrowseRootDirectoryClicked()
 }
 
 FReply SNewWorkspaceWindow::OnCreateClicked()
-{	
+{
 	return FReply::Handled();
 }
 
diff --git a/Engine/Source/Programs/SlateUGS/Private/Widgets/SNewWorkspaceWindow.h b/Engine/Source/Programs/SlateUGS/Private/Widgets/SNewWorkspaceWindow.h
--- a/Engine/Source/Programs/SlateUGS/Private/Widgets/SNewWorkspaceWindow.h
+++ b/Engine/Source/Programs/SlateUGS/Private/Widgets/SNewWorkspaceWindow.h
@@ -11,10 +11,9 @@ class SNewWorkspaceWindow final : public SWindow
 {
 public:
 	SLATE_BEGIN_ARGS(SNewWorkspaceWindow) {}
-	SLATE_ARGUMENT(UGSTab*, Tab)
-SLATE_END_ARGS()
+	SLATE_END_ARGS()
 
-void Construct(const FArguments& InArgs);
+	void Construct(const FArguments& InArgs, UGSTab* InTab);
 
 private:
 	FReply OnBrowseStreamClicked();
@@ -23,7 +22,7 @@ private:
 	FReply OnCreateClicked();
 	FReply OnCancelClicked();
 
-	TSharedPtr<SEditableTextBox> LocalFileText = nullptr;
+	TSharedPtr<SEditableTextBox> LocalFileText;
 	FString WorkspacePathText;
 
 	UGSTab* Tab = nullptr;
diff --git a/Engine/Source/Programs/SlateUGS/Private/Widgets/SScheduledSyncWindow.cpp b/Engine/Source/Programs/SlateUGS/Private/Widgets/SScheduledSyncWindow.cpp
--- a/Engine/Source/Programs/SlateUGS/Private/Widgets/SScheduledSyncWindow.cpp
+++ b/Engine/Source/Programs/SlateUGS/Private/Widgets/SScheduledSyncWindow.cpp
@@ -21,9 +21,9 @@
 // Todo: remove checkbox and make the positive action button the thing that schedules the sync
 // Todo: show UI if a scheduled sync already exists and provide the ability to cancel it
 
-void SScheduledSyncWindow::Construct(const FArguments& InArgs)
+void SScheduledSyncWindow::Construct(const FArguments& InArgs, UGSTab* InTab)
 {
-	Tab = InArgs._Tab;
+	Tab = InTab;
 	UserSettings = Tab->GetUserSettings();
 
 	SWindow::Construct(SWindow::FArguments()
diff --git a/Engine/Source/Programs/SlateUGS/Private/Widgets/SScheduledSyncWindow.h b/Engine/Source/Programs/SlateUGS/Private/Widgets/SScheduledSyncWindow.h
--- a/Engine/Source/Programs/SlateUGS/Private/Widgets/SScheduledSyncWindow.h
+++ b/Engine/Source/Programs/SlateUGS/Private/Widgets/SScheduledSyncWindow.h
@@ -12,14 +12,13 @@ namespace UGSCore
 	struct FUserSettings;
 }
 
-class SScheduledSyncWindow : public SWindow
+class SScheduledSyncWindow final : public SWindow
 {
 public:
 	SLATE_BEGIN_ARGS(SScheduledSyncWindow) {}
-		SLATE_ARGUMENT(UGSTab*, Tab)
 	SLATE_END_ARGS()
 
-	void Construct(const FArguments& InArgs);
+	void Construct(const FArguments& InArgs, UGSTab* InTab);
 
 private:
 	/** Handles getting the text color of the editable text box. */
@@ -41,6 +40,6 @@ private:
 	bool bInputValid = false;
 	FString Input;
 
-	UGSTab* Tab;
+	UGSTab* Tab = nullptr;
 	TSharedPtr<UGSCore::FUserSettings> UserSettings;
 };
diff --git a/Engine/Source/Programs/SlateUGS/Private/Widgets/SSelectStreamWindow.cpp b/Engine/Source/Programs/SlateUGS/Private/Widgets/SSelectStreamWindow.cpp
new file mode 100644
--- /dev/null
+++ b/Engine/Source/Programs/SlateUGS/Private/Widgets/SSelectStreamWindow.cpp
@@ -0,0 +1,94 @@
+// Copyright Epic Games, Inc. All Rights Reserved.
+
+#include "SSelectStreamWindow.h"
+
+#include "SlateUGSStyle.h"
+#include "Framework/Application/SlateApplication.h"
+
+#include "UGSTab.h"
+#include "SGameSyncTab.h"
+#include "SPrimaryButton.h"
+#include "Widgets/Layout/SHeader.h"
+
+#define LOCTEXT_NAMESPACE "UGSNewWorkspaceWindow"
+
+struct FStreamNode
+{
+	FStreamNode(const FString& InLabel, bool bInIsStream)
+		: Label(FText::FromString(InLabel))
+		, bIsStream(bInIsStream) {}
+
+	virtual ~FStreamNode() {}
+
+	FText Label;
+	bool bIsStream;
+	TArray<TSharedPtr<FStreamNode>> Children;
+};
+
+void SSelectStreamWindow::Construct(const FArguments& InArgs, UGSTab* InTab)
+{
+	Tab = InTab;
+
+	SWindow::Construct(SWindow::FArguments()
+	.Title(LOCTEXT("WindowTitle", "Select Stream"))
+	.SizingRule(ESizingRule::FixedSize)
+	.ClientSize(FVector2D(600, 500))
+	[
+		SNew(SBox)
+		.Padding(30.0f, 15.0f, 30.0f, 0.0f)
+		[
+			SNew(SVerticalBox)
+			+SVerticalBox::Slot()
+			.AutoHeight()
+			[
+				SNew(SEditableTextBox)
+				// Todo: change hint and enable when filters are supported
+				.HintText(LOCTEXT("FilterHint", "Filter (under construction, does not work yet)"))
+				.IsEnabled(false)
+			]
+			+SVerticalBox::Slot()
+			[
+				SNew(STreeView<TSharedPtr<FStreamNode>>)
+				.TreeItemsSource(&StreamsTree)
+			]
+			+SVerticalBox::Slot()
+			.AutoHeight()
+			.VAlign(VAlign_Bottom)
+			[
+				SNew(SBox)
+				.HAlign(HAlign_Right)
+				.Padding(10.0f)
+				[
+					SNew(SHorizontalBox)
+					+SHorizontalBox::Slot()
+					.Padding(0.0f, 0.0f, 10.0f, 0.0f)
+					[
+						SNew(SPrimaryButton)
+						.Text(LOCTEXT("OkButtonText", "Ok"))
+						.OnClicked(this, &SSelectStreamWindow::OnOkClicked)
+					]
+					+SHorizontalBox::Slot()
+					.AutoWidth()
+					[
+						SNew(SButton)
+						.Text(LOCTEXT("CancelButtonText", "Cancel"))
+						.OnClicked(this, &SSelectStreamWindow::OnCancelClicked)
+					]
+				]
+			]
+		]
+	]);
+}
+
+FReply SSelectStreamWindow::OnOkClicked()
+{
+	return FReply::Handled();
+}
+
+FReply SSelectStreamWindow::OnCancelClicked()
+{
+	FSlateApplication::Get().FindWidgetWindow(AsShared())->RequestDestroyWindow();
+	return FReply::Handled();
+}
+
+#undef LOCTEXT_NAMESPACE
diff --git a/Engine/Source/Programs/SlateUGS/Private/Widgets/SSelectStreamWindow.h b/Engine/Source/Programs/SlateUGS/Private/Widgets/SSelectStreamWindow.h
new file mode 100644
--- /dev/null
+++ b/Engine/Source/Programs/SlateUGS/Private/Widgets/SSelectStreamWindow.h
@@ -0,0 +1,27 @@
+// Copyright Epic Games, Inc. All Rights Reserved.
+
+#pragma once
+
+#include "Widgets/SWindow.h"
+#include "Widgets/Input/SEditableTextBox.h"
+
+class UGSTab;
+struct FStreamNode;
+
+class SSelectStreamWindow final : public SWindow
+{
+public:
+	SLATE_BEGIN_ARGS(SSelectStreamWindow) {}
+	SLATE_END_ARGS()
+
+	void Construct(const FArguments& InArgs, UGSTab* InTab);
+
+private:
+	FReply OnOkClicked();
+	FReply OnCancelClicked();
+
+	TSharedPtr<SEditableTextBox> FilterText;
+	TArray<TSharedPtr<FStreamNode>> StreamsTree;
+
+	UGSTab* Tab = nullptr;
+};
diff --git a/Engine/Source/Programs/SlateUGS/Private/Widgets/SSettingsWindow.cpp b/Engine/Source/Programs/SlateUGS/Private/Widgets/SSettingsWindow.cpp
--- a/Engine/Source/Programs/SlateUGS/Private/Widgets/SSettingsWindow.cpp
+++ b/Engine/Source/Programs/SlateUGS/Private/Widgets/SSettingsWindow.cpp
@@ -15,9 +15,9 @@
 
 // Todo: Make settings only save when the positive action button is hit
 
-void SSettingsWindow::Construct(const FArguments& InArgs)
+void SSettingsWindow::Construct(const FArguments& InArgs, UGSTab* InTab)
 {
-	Tab = InArgs._Tab;
+	Tab = InTab;
 	UserSettings = Tab->GetUserSettings();
 
 	SWindow::Construct(SWindow::FArguments()
diff --git a/Engine/Source/Programs/SlateUGS/Private/Widgets/SSettingsWindow.h b/Engine/Source/Programs/SlateUGS/Private/Widgets/SSettingsWindow.h
--- a/Engine/Source/Programs/SlateUGS/Private/Widgets/SSettingsWindow.h
+++ b/Engine/Source/Programs/SlateUGS/Private/Widgets/SSettingsWindow.h
@@ -13,14 +13,13 @@ namespace UGSCore
 	struct FUserSettings;
 }
 
-class SSettingsWindow : public SWindow
+class SSettingsWindow final: public SWindow
 {
 public:
 	SLATE_BEGIN_ARGS(SSettingsWindow) {}
-		SLATE_ARGUMENT(UGSTab*, Tab)
 	SLATE_END_ARGS()
 
-	void Construct(const FArguments& InArgs);
+	void Construct(const FArguments& InArgs, UGSTab* InTab);
 
 private:
 	/** Handle the check box enable or disabling for AfterSync settings */
@@ -39,6 +38,6 @@ private:
 	FReply OnOkClicked();
 	FReply OnCancelClicked();
 
-	UGSTab* Tab;
+	UGSTab* Tab = nullptr;
 	TSharedPtr<UGSCore::FUserSettings> UserSettings;
 };
diff --git a/Engine/Source/Programs/SlateUGS/Private/Widgets/SSyncFilterWindow.cpp b/Engine/Source/Programs/SlateUGS/Private/Widgets/SSyncFilterWindow.cpp
--- a/Engine/Source/Programs/SlateUGS/Private/Widgets/SSyncFilterWindow.cpp
+++ b/Engine/Source/Programs/SlateUGS/Private/Widgets/SSyncFilterWindow.cpp
@@ -8,7 +8,6 @@
 #include "Widgets/Layout/SUniformGridPanel.h"
 #include "Widgets/Layout/SHeader.h"
 #include "SPrimaryButton.h"
-#include "SSimpleButton.h"
 #include "SPopupTextWindow.h"
 
 #define LOCTEXT_NAMESPACE "SSyncFilterWindow"
@@ -77,9 +76,9 @@ void SSyncFilterWindow::ConstructCustomSyncViewTextBoxes()
 	.Text(FText::FromString(FString::Join(Tab->GetSyncViews(SyncCategoryType::AllWorkspaces), TEXT("\n"))));
 }
 
-void SSyncFilterWindow::Construct(const FArguments& InArgs)
+void SSyncFilterWindow::Construct(const FArguments& InArgs, UGSTab* InTab)
 {
-	Tab = InArgs._Tab;
+	Tab = InTab;
 
 	ConstructSyncFilters();
 	ConstructCustomSyncViewTextBoxes();
diff --git a/Engine/Source/Programs/SlateUGS/Private/Widgets/SSyncFilterWindow.h b/Engine/Source/Programs/SlateUGS/Private/Widgets/SSyncFilterWindow.h
--- a/Engine/Source/Programs/SlateUGS/Private/Widgets/SSyncFilterWindow.h
+++ b/Engine/Source/Programs/SlateUGS/Private/Widgets/SSyncFilterWindow.h
@@ -11,14 +11,13 @@
 
 class UGSTab;
 
-class SSyncFilterWindow : public SWindow
+class SSyncFilterWindow final : public SWindow
 {
 public:
 	SLATE_BEGIN_ARGS(SSyncFilterWindow) {}
-		SLATE_ARGUMENT(UGSTab*, Tab)
 	SLATE_END_ARGS()
 
-	void Construct(const FArguments& InArgs);
+	void Construct(const FArguments& InArgs, UGSTab* InTab);
 
 private:
 	FReply OnShowCombinedFilterClicked();
@@ -29,7 +28,7 @@ private:
 	void ConstructSyncFilters();
 	void ConstructCustomSyncViewTextBoxes();
 
-	UGSTab* Tab;
+	UGSTab* Tab = nullptr;
 	TArray<UGSCore::FWorkspaceSyncCategory> WorkspaceCategoriesCurrent;
 	TArray<UGSCore::FWorkspaceSyncCategory> WorkspaceCategoriesAll;
 
diff --git a/Engine/Source/Programs/SlateUGS/Private/Widgets/SWorkspaceWindow.cpp b/Engine/Source/Programs/SlateUGS/Private/Widgets/SWorkspaceWindow.cpp
--- a/Engine/Source/Programs/SlateUGS/Private/Widgets/SWorkspaceWindow.cpp
+++ b/Engine/Source/Programs/SlateUGS/Private/Widgets/SWorkspaceWindow.cpp
@@ -16,9 +16,9 @@
 
 #define LOCTEXT_NAMESPACE "UGSWorkspaceWindow"
 
-void SWorkspaceWindow::Construct(const FArguments& InArgs)
+void SWorkspaceWindow::Construct(const FArguments& InArgs, UGSTab* InTab)
 {
-	Tab = InArgs._Tab;
+	Tab = InTab;
 
 	SWindow::Construct(SWindow::FArguments()
 	.Title(LOCTEXT("WindowTitle", "Open Project"))
@@ -266,7 +266,7 @@ FReply SWorkspaceWindow::OnBrowseClicked()
 
 FReply SWorkspaceWindow::OnNewClicked()
 {
-	FSlateApplication::Get().AddModalWindow(SNew(SNewWorkspaceWindow).Tab(Tab), SharedThis(this), false);
+	FSlateApplication::Get().AddModalWindow(SNew(SNewWorkspaceWindow, Tab), SharedThis(this), false);
 
 	return FReply::Handled();
 }
diff --git a/Engine/Source/Programs/SlateUGS/Private/Widgets/SWorkspaceWindow.h b/Engine/Source/Programs/SlateUGS/Private/Widgets/SWorkspaceWindow.h
--- a/Engine/Source/Programs/SlateUGS/Private/Widgets/SWorkspaceWindow.h
+++ b/Engine/Source/Programs/SlateUGS/Private/Widgets/SWorkspaceWindow.h
@@ -7,14 +7,13 @@
 
 class UGSTab;
 
-class SWorkspaceWindow : public SWindow
+class SWorkspaceWindow final : public SWindow
 {
 public:
 	SLATE_BEGIN_ARGS(SWorkspaceWindow) {}
-		SLATE_ARGUMENT(UGSTab*, Tab)
 	SLATE_END_ARGS()
 
-	void Construct(const FArguments& InArgs);
+	void Construct(const FArguments& InArgs, UGSTab* InTab);
 
 private:
 	FReply OnOkClicked();
@@ -27,5 +26,5 @@ private:
 	TSharedPtr<SEditableTextBox> LocalFileText = nullptr;
 	FString WorkspacePathText;
 
-	UGSTab* Tab;
+	UGSTab* Tab = nullptr;
 };
