Fix SlateUGS compiler error on Windows

diff --git a/Engine/Source/Programs/SlateUGS/Private/Widgets/SSelectStreamWindow.cpp b/Engine/Source/Programs/SlateUGS/Private/Widgets/SSelectStreamWindow.cpp
--- a/Engine/Source/Programs/SlateUGS/Private/Widgets/SSelectStreamWindow.cpp
+++ b/Engine/Source/Programs/SlateUGS/Private/Widgets/SSelectStreamWindow.cpp
@@ -9,6 +9,7 @@
 #include "SGameSyncTab.h"
 #include "SPrimaryButton.h"
 #include "Widgets/Layout/SHeader.h"
+#include "Widgets/Views/STreeView.h"
 
 #define LOCTEXT_NAMESPACE "UGSNewWorkspaceWindow"
 
