Temporarily re-labeled cancel buttons to close on menus which save settings even when you hit close

diff --git a/Engine/Source/Programs/SlateUGS/Private/Widgets/SScheduledSyncWindow.cpp b/Engine/Source/Programs/SlateUGS/Private/Widgets/SScheduledSyncWindow.cpp
--- a/Engine/Source/Programs/SlateUGS/Private/Widgets/SScheduledSyncWindow.cpp
+++ b/Engine/Source/Programs/SlateUGS/Private/Widgets/SScheduledSyncWindow.cpp
@@ -115,7 +115,7 @@ void SScheduledSyncWindow::Construct(const FArguments& InArgs, UGSTab* InTab)
 				[
 					SNew(SButton)
 					.HAlign(HAlign_Center)
-					.Text(LOCTEXT("CancelButtonText", "Cancel"))
+					.Text(LOCTEXT("CancelButtonText", "Close")) // Todo: temporary fix, change the text back to "Cancel" when canceling actually cancels changes
 					.OnClicked(this, &SScheduledSyncWindow::OnCancelClicked)
 				]
 			]
diff --git a/Engine/Source/Programs/SlateUGS/Private/Widgets/SSettingsWindow.cpp b/Engine/Source/Programs/SlateUGS/Private/Widgets/SSettingsWindow.cpp
--- a/Engine/Source/Programs/SlateUGS/Private/Widgets/SSettingsWindow.cpp
+++ b/Engine/Source/Programs/SlateUGS/Private/Widgets/SSettingsWindow.cpp
@@ -181,7 +181,7 @@ void SSettingsWindow::Construct(const FArguments& InArgs, UGSTab* InTab)
 				[
 					SNew(SButton)
 					.HAlign(HAlign_Center)
-					.Text(LOCTEXT("CancelButtonText", "Cancel"))
+					.Text(LOCTEXT("CancelButtonText", "Close")) // Todo: temporary fix, change the text back to "Cancel" when canceling actually cancels changes
 					.OnClicked(this, &SSettingsWindow::OnCancelClicked)
 				]
 			]
