Select stream window works and has all features except the loading screen and the search filter (lower prio tasks)

diff --git a/Engine/Source/Programs/SlateUGS/Private/UGSTab.cpp b/Engine/Source/Programs/SlateUGS/Private/UGSTab.cpp
--- a/Engine/Source/Programs/SlateUGS/Private/UGSTab.cpp
+++ b/Engine/Source/Programs/SlateUGS/Private/UGSTab.cpp
@@ -204,9 +204,22 @@ bool UGSTab::ShouldSyncPrecompiledEditor() const
 
 TArray<FString> UGSTab::GetAllStreamNames() const
 {
-	TArray<FString> Result;
 	FEvent* AbortEvent = FPlatformProcess::GetSynchEventFromPool(true);
-	PerforceClient->FindStreams("//*/*", Result, AbortEvent, *MakeShared<FLineWriter>());
+	
+	TArray<FString> Result;
+	const TSharedRef<UGSCore::FModalTaskResult> TaskResult = ExecuteModalTask(
+		FSlateApplication::Get().GetActiveModalWindow(),
+		PerforceClient->FindStreams("//*/*", Result, AbortEvent, *MakeShared<FLineWriter>()),
+		LOCTEXT("OpeningProjectTitle", "Opening Project"),
+		LOCTEXT("OpeningProjectCaption", "Opening project, please wait..."));
+	
+	if (TaskResult->Failed())
+	{
+		FMessageDialog::Open(EAppMsgType::Ok, TaskResult->GetMessage());
+		return TArray<FString>();
+	}
+
+	Result.Sort();
 
 	return Result;
 }
diff --git a/Engine/Source/Programs/SlateUGS/Private/Widgets/SNewWorkspaceWindow.cpp b/Engine/Source/Programs/SlateUGS/Private/Widgets/SNewWorkspaceWindow.cpp
--- a/Engine/Source/Programs/SlateUGS/Private/Widgets/SNewWorkspaceWindow.cpp
+++ b/Engine/Source/Programs/SlateUGS/Private/Widgets/SNewWorkspaceWindow.cpp
@@ -70,7 +70,7 @@ void SNewWorkspaceWindow::Construct(const FArguments& InArgs, UGSTab* InTab)
 						SNew(SHorizontalBox)
 						+SHorizontalBox::Slot()
 						[
-							SNew(SEditableTextBox)
+							SAssignNew(StreamTextBox, SEditableTextBox)
 						]
 						+SHorizontalBox::Slot()
 						.FillWidth(0.225f)
@@ -87,7 +87,7 @@ void SNewWorkspaceWindow::Construct(const FArguments& InArgs, UGSTab* InTab)
 						SNew(SHorizontalBox)
 						+SHorizontalBox::Slot()
 						[
-							SNew(SEditableTextBox)
+							SAssignNew(RootDirTextBox, SEditableTextBox)
 						]
 						+SHorizontalBox::Slot()
 						.FillWidth(0.225f)
@@ -100,7 +100,7 @@ void SNewWorkspaceWindow::Construct(const FArguments& InArgs, UGSTab* InTab)
 					]
 					+SVerticalBox::Slot()
 					[
-						SNew(SEditableTextBox)
+						SAssignNew(FileNameTextBox, SEditableTextBox)
 					]
 				]
 			]
@@ -134,7 +134,10 @@ void SNewWorkspaceWindow::Construct(const FArguments& InArgs, UGSTab* InTab)
 
 FReply SNewWorkspaceWindow::OnBrowseStreamClicked()
 {
-	FSlateApplication::Get().AddModalWindow(SNew(SSelectStreamWindow, Tab), SharedThis(this), false);
+	FString OutStreamPath;
+	FSlateApplication::Get().AddModalWindow(SNew(SSelectStreamWindow, Tab, &OutStreamPath), SharedThis(this), false);
+
+	StreamTextBox->SetText(FText::FromString(*OutStreamPath));
 
 	return FReply::Handled();
 }
diff --git a/Engine/Source/Programs/SlateUGS/Private/Widgets/SNewWorkspaceWindow.h b/Engine/Source/Programs/SlateUGS/Private/Widgets/SNewWorkspaceWindow.h
--- a/Engine/Source/Programs/SlateUGS/Private/Widgets/SNewWorkspaceWindow.h
+++ b/Engine/Source/Programs/SlateUGS/Private/Widgets/SNewWorkspaceWindow.h
@@ -22,7 +22,9 @@ private:
 	FReply OnCreateClicked();
 	FReply OnCancelClicked();
 
-	TSharedPtr<SEditableTextBox> LocalFileText;
+	TSharedPtr<SEditableTextBox> StreamTextBox = nullptr;
+	TSharedPtr<SEditableTextBox> RootDirTextBox = nullptr;
+	TSharedPtr<SEditableTextBox> FileNameTextBox = nullptr;
 	FString WorkspacePathText;
 
 	UGSTab* Tab = nullptr;
diff --git a/Engine/Source/Programs/SlateUGS/Private/Widgets/SSelectStreamWindow.cpp b/Engine/Source/Programs/SlateUGS/Private/Widgets/SSelectStreamWindow.cpp
--- a/Engine/Source/Programs/SlateUGS/Private/Widgets/SSelectStreamWindow.cpp
+++ b/Engine/Source/Programs/SlateUGS/Private/Widgets/SSelectStreamWindow.cpp
@@ -15,20 +15,25 @@
 
 struct FStreamNode
 {
-	FStreamNode(const FString& InLabel, bool bInIsStream)
+	FStreamNode(const FString& InLabel, const FString& InFullStreamPath, bool bInIsStream)
 		: Label(FText::FromString(InLabel))
+		, FullStreamPath(InFullStreamPath)
 		, bIsStream(bInIsStream) {}
 
 	virtual ~FStreamNode() {}
 
 	FText Label;
+	FString FullStreamPath;
 	bool bIsStream;
 	TArray<TSharedPtr<FStreamNode>> Children;
 };
 
-void SSelectStreamWindow::Construct(const FArguments& InArgs, UGSTab* InTab)
+void SSelectStreamWindow::Construct(const FArguments& InArgs, UGSTab* InTab, FString* OutSelectedStreamPath)
 {
 	Tab = InTab;
+	SelectedStreamPath = OutSelectedStreamPath;
+
+	PopulateStreamsTree();
 
 	SWindow::Construct(SWindow::FArguments()
 	.Title(LOCTEXT("WindowTitle", "Select Stream"))
@@ -48,9 +53,13 @@ void SSelectStreamWindow::Construct(const FArguments& InArgs, UGSTab* InTab)
 				.IsEnabled(false)
 			]
 			+SVerticalBox::Slot()
+			.Padding(0.0f, 15.0f, 0.0f, 0.0f)
 			[
-				SNew(STreeView<TSharedPtr<FStreamNode>>)
+				SNew(STreeView<TSharedRef<FStreamNode>>)
 				.TreeItemsSource(&StreamsTree)
+				.OnGenerateRow(this, &SSelectStreamWindow::OnGenerateRow)
+				.OnGetChildren(this, &SSelectStreamWindow::OnGetChildren)
+				.OnSelectionChanged(this, &SSelectStreamWindow::OnSelectionChanged)
 			]
 			+SVerticalBox::Slot()
 			.AutoHeight()
@@ -67,6 +76,7 @@ void SSelectStreamWindow::Construct(const FArguments& InArgs, UGSTab* InTab)
 						SNew(SPrimaryButton)
 						.Text(LOCTEXT("OkButtonText", "Ok"))
 						.OnClicked(this, &SSelectStreamWindow::OnOkClicked)
+						.IsEnabled(this, &SSelectStreamWindow::IsOkButtonEnabled)
 					]
 					+SHorizontalBox::Slot()
 					.AutoWidth()
@@ -81,8 +91,72 @@ void SSelectStreamWindow::Construct(const FArguments& InArgs, UGSTab* InTab)
 	]);
 }
 
+TSharedRef<ITableRow> SSelectStreamWindow::OnGenerateRow(TSharedRef<FStreamNode> InItem, const TSharedRef<STableViewBase>& InOwnerTable)
+{
+	return SNew(STableRow<TSharedRef<FStreamNode>>, InOwnerTable)
+	[
+		SNew(SHorizontalBox)
+		+SHorizontalBox::Slot()
+		.AutoWidth()
+		.HAlign(HAlign_Left)
+		.Padding(0.0f, 0.0f, 5.0f, 0.0f)
+		// [
+		// 	SNew(SImage)
+		// 	// Todo: replace text with keys to actuall perforce depot/stream icons
+		// 	.Image(FSlateUGSStyle::Get().GetBrush(InItem->bIsStream ? "FILL ME IN" : "FILL ME IN TOO))
+		// ]
+		+SHorizontalBox::Slot()
+		.VAlign(VAlign_Center)
+		[
+			SNew(STextBlock)
+			.Text(InItem->Label)
+		]
+	];
+}
+
+void SSelectStreamWindow::OnGetChildren(TSharedRef<FStreamNode> InItem, TArray<TSharedRef<FStreamNode>>& OutChildren)
+{
+	for (int32 ChildIndex = 0; ChildIndex < InItem->Children.Num(); ++ChildIndex)
+	{
+		TSharedPtr<FStreamNode> Child = InItem->Children[ChildIndex];
+		if (Child.IsValid())
+		{
+			OutChildren.Add(Child.ToSharedRef());
+		}
+	}
+}
+
+void SSelectStreamWindow::PopulateStreamsTree()
+{
+	TArray<FString> Streams = Tab->GetAllStreamNames();
+	FString CurrentRoot = "";
+
+	for (const FString& Stream : Streams)
+	{
+		TArray<FString> StreamPath;
+		Stream.RightChop(2).ParseIntoArray(StreamPath, TEXT("/"));
+
+		if (StreamPath[0] != CurrentRoot)
+		{
+			CurrentRoot = StreamPath[0];
+			StreamsTree.Add(MakeShareable(new FStreamNode(StreamPath[0], Stream, false)));
+		}
+
+		FString CurrentChild = Stream.RightChop(CurrentRoot.Len() + 3);
+		StreamsTree.Last()->Children.Add(MakeShareable(new FStreamNode(CurrentChild, Stream, true)));
+	}
+}
+
+void SSelectStreamWindow::OnSelectionChanged(TSharedPtr<FStreamNode> NewSelection, ESelectInfo::Type SelectInfo)
+{
+	SelectedStream = NewSelection;
+}
+
 FReply SSelectStreamWindow::OnOkClicked()
 {
+	*SelectedStreamPath = SelectedStream->FullStreamPath;
+
+	FSlateApplication::Get().FindWidgetWindow(AsShared())->RequestDestroyWindow();
 	return FReply::Handled();
 }
 
@@ -92,4 +166,9 @@ FReply SSelectStreamWindow::OnCancelClicked()
 	return FReply::Handled();
 }
 
+bool SSelectStreamWindow::IsOkButtonEnabled() const
+{
+	return SelectedStream.IsValid() && SelectedStream->bIsStream;
+}
+
 #undef LOCTEXT_NAMESPACE
diff --git a/Engine/Source/Programs/SlateUGS/Private/Widgets/SSelectStreamWindow.h b/Engine/Source/Programs/SlateUGS/Private/Widgets/SSelectStreamWindow.h
--- a/Engine/Source/Programs/SlateUGS/Private/Widgets/SSelectStreamWindow.h
+++ b/Engine/Source/Programs/SlateUGS/Private/Widgets/SSelectStreamWindow.h
@@ -14,14 +14,23 @@ public:
 	SLATE_BEGIN_ARGS(SSelectStreamWindow) {}
 	SLATE_END_ARGS()
 
-	void Construct(const FArguments& InArgs, UGSTab* InTab);
+	void Construct(const FArguments& InArgs, UGSTab* InTab, FString* OutSelectedStreamPath);
 
 private:
+	TSharedPtr<SEditableTextBox> FilterText;
+	TSharedPtr<FStreamNode> SelectedStream;
+	FString* SelectedStreamPath = nullptr;
+
+	TArray<TSharedRef<FStreamNode>> StreamsTree;
+	void PopulateStreamsTree();
+	TSharedRef<ITableRow> OnGenerateRow(TSharedRef<FStreamNode> InItem, const TSharedRef<STableViewBase>& InOwnerTable);
+	void OnGetChildren(TSharedRef<FStreamNode> InItem, TArray<TSharedRef<FStreamNode>>& OutChildren);
+	void OnSelectionChanged(TSharedPtr<FStreamNode> NewSelection, ESelectInfo::Type SelectInfo);
+
 	FReply OnOkClicked();
 	FReply OnCancelClicked();
 
-	TSharedPtr<SEditableTextBox> FilterText;
-	TArray<TSharedPtr<FStreamNode>> StreamsTree;
+	bool IsOkButtonEnabled() const;
 
 	UGSTab* Tab = nullptr;
 };
