Removing accidental early submission

diff --git a/Engine/Source/Programs/SlateUGS/Private/UGSTab.cpp b/Engine/Source/Programs/SlateUGS/Private/UGSTab.cpp
--- a/Engine/Source/Programs/SlateUGS/Private/UGSTab.cpp
+++ b/Engine/Source/Programs/SlateUGS/Private/UGSTab.cpp
@@ -204,22 +204,9 @@ bool UGSTab::ShouldSyncPrecompiledEditor() const
 
 TArray<FString> UGSTab::GetAllStreamNames() const
 {
-	FEvent* AbortEvent = FPlatformProcess::GetSynchEventFromPool(true);
-	
 	TArray<FString> Result;
-	const TSharedRef<UGSCore::FModalTaskResult> TaskResult = ExecuteModalTask(
-		FSlateApplication::Get().GetActiveModalWindow(),
-		PerforceClient->FindStreams("//*/*", Result, AbortEvent, *MakeShared<FLineWriter>()),
-		LOCTEXT("OpeningProjectTitle", "Opening Project"),
-		LOCTEXT("OpeningProjectCaption", "Opening project, please wait..."));
-	
-	if (TaskResult->Failed())
-	{
-		FMessageDialog::Open(EAppMsgType::Ok, TaskResult->GetMessage());
-		return TArray<FString>();
-	}
-
-	Result.Sort();
+	FEvent* AbortEvent = FPlatformProcess::GetSynchEventFromPool(true);
+	PerforceClient->FindStreams("//*/*", Result, AbortEvent, *MakeShared<FLineWriter>());
 
 	return Result;
 }
