Add loading modal dialog window when selecting stream for new workspace

diff --git a/Engine/Source/Programs/SlateUGS/Private/UGSTab.cpp b/Engine/Source/Programs/SlateUGS/Private/UGSTab.cpp
--- a/Engine/Source/Programs/SlateUGS/Private/UGSTab.cpp
+++ b/Engine/Source/Programs/SlateUGS/Private/UGSTab.cpp
@@ -4,11 +4,13 @@
 
 #include "UGSLog.h"
 #include "ChangeInfo.h"
+#include "FindStreamsTask.h"
 #include "UGSTabManager.h"
 
 #include "UGSCore/Utility.h"
 #include "UGSCore/BuildStep.h"
 #include "UGSCore/DetectProjectSettingsTask.h"
+#include "UGSCore/FindStreamsTask.h"
 #include "UGSCore/PerforceMonitor.h"
 #include "UGSCore/EventMonitor.h"
 
@@ -202,11 +204,23 @@ bool UGSTab::ShouldSyncPrecompiledEditor() const
 	return UserSettings->bSyncPrecompiledEditor && PerforceMonitor->HasZippedBinaries();
 }
 
-TArray<FString> UGSTab::GetAllStreamNames() const
+TArray<FString> UGSTab::GetAllStreamNames()
 {
 	TArray<FString> Result;
-	FEvent* AbortEvent = FPlatformProcess::GetSynchEventFromPool(true);
-	PerforceClient->FindStreams("//*/*", Result, AbortEvent, *MakeShared<FLineWriter>());
+
+	const TSharedRef<UGSCore::FModalTaskResult> TaskResult = ExecuteModalTask(
+		FSlateApplication::Get().GetActiveModalWindow(),
+		MakeShared<UGSCore::FFindStreamsTask>(PerforceClient.ToSharedRef(), MakeShared<FLineWriter>(), Result, "//*/*"),
+		LOCTEXT("OpeningProjectTitle", "Finding Streams"),
+		LOCTEXT("OpeningProjectCaption", "Finding streams, please wait..."));
+	
+	if (TaskResult->Failed())
+	{
+		FMessageDialog::Open(EAppMsgType::Ok, TaskResult->GetMessage());
+		return TArray<FString>();
+	}
+
+	Result.Sort();
 
 	return Result;
 }
@@ -611,7 +625,7 @@ const TArray<FString>& UGSTab::GetCombinedSyncFilter() const
 	return CombinedSyncFilter;
 }
 
-UGSTabManager* UGSTab::GetTabManager()
+UGSTabManager* UGSTab::GetTabManager() const
 {
 	return TabManager;
 }
@@ -745,4 +759,4 @@ bool UGSTab::SetupWorkspace()
 	return true;
 }
 
-#undef LOCTEXT_NAMESPACE
+#undef LOCTEXT_NAMESPACE
\ No newline at end of file
diff --git a/Engine/Source/Programs/SlateUGS/Private/UGSTab.h b/Engine/Source/Programs/SlateUGS/Private/UGSTab.h
--- a/Engine/Source/Programs/SlateUGS/Private/UGSTab.h
+++ b/Engine/Source/Programs/SlateUGS/Private/UGSTab.h
@@ -2,9 +2,9 @@
 
 #pragma once
 
+#include "FindStreamsTask.h"
 #include "Widgets/SEmptyTab.h"
 #include "Widgets/SGameSyncTab.h"
-#include "Widgets/SWorkspaceWindow.h"
 #include "Widgets/Docking/SDockTab.h"
 
 #include "UGSCore/Workspace.h"
@@ -66,14 +66,13 @@ public:
 	// Accessors
 	bool IsSyncing() const;
 	FString GetSyncProgress() const;
-	const TArray<FString>& GetSyncFilters() const;
 	const TArray<FString>& GetCombinedSyncFilter() const;
 	TArray<UGSCore::FWorkspaceSyncCategory> GetSyncCategories(SyncCategoryType CategoryType) const;
 	TArray<FString> GetSyncViews(SyncCategoryType CategoryType) const;
-	UGSTabManager* GetTabManager();
+	UGSTabManager* GetTabManager() const;
 	TSharedPtr<UGSCore::FUserSettings> GetUserSettings() const;
 	bool ShouldSyncPrecompiledEditor() const;
-	TArray<FString> GetAllStreamNames() const;
+	TArray<FString> GetAllStreamNames();
 
 	void UpdateGameTabBuildList();
 	void RefreshBuildList();
diff --git a/Engine/Source/Programs/UGSCore/FindStreamsTask.cpp b/Engine/Source/Programs/UGSCore/FindStreamsTask.cpp
new file mode 100644
--- /dev/null
+++ b/Engine/Source/Programs/UGSCore/FindStreamsTask.cpp
@@ -0,0 +1,46 @@
+// Copyright Epic Games, Inc. All Rights Reserved.
+
+#include "FindStreamsTask.h"
+
+#define LOCTEXT_NAMESPACE "UnrealGameSync"
+
+namespace UGSCore
+{
+
+	FFindStreamsTask::FFindStreamsTask(
+		TSharedRef<FPerforceConnection> InPerforce,
+		TSharedRef<FLineBasedTextWriter> InLog,
+		TArray<FString>& OutStreams,
+		const FString& StreamFilter)
+		: Perforce(MoveTemp(InPerforce))
+		, Log(MoveTemp(InLog))
+		, Streams(OutStreams)
+		, Filter(StreamFilter)
+	{}
+
+	FFindStreamsTask::~FFindStreamsTask()
+	{}
+
+	TSharedRef<FModalTaskResult> FFindStreamsTask::Run(FEvent* AbortEvent)
+	{
+		// TODO enable these if we want
+		//try
+		{
+			return RunInternal(AbortEvent);
+		}
+		//catch(FAbortException)
+		{
+			//return FModalTaskResult::Aborted();
+		}
+	}
+
+	TSharedRef<FModalTaskResult> FFindStreamsTask::RunInternal(FEvent* AbortEvent)
+	{
+		return Perforce->FindStreams(Filter, Streams, AbortEvent, *Log)
+			? FModalTaskResult::Success()
+			: FModalTaskResult::Failure(LOCTEXT("PerforceTaskFailure", "Failed to find streams"));
+	}
+
+#undef LOCTEXT_NAMESPACE
+
+} // namespace UGSCore
diff --git a/Engine/Source/Programs/UGSCore/FindStreamsTask.h b/Engine/Source/Programs/UGSCore/FindStreamsTask.h
new file mode 100644
--- /dev/null
+++ b/Engine/Source/Programs/UGSCore/FindStreamsTask.h
@@ -0,0 +1,36 @@
+// Copyright Epic Games, Inc. All Rights Reserved.
+
+#pragma once
+
+#include "CoreMinimal.h"
+#include "ModalTask.h"
+#include "Perforce.h"
+#include "OutputAdapters.h"
+
+namespace UGSCore
+{
+
+class FFindStreamsTask : public IModalTask
+{
+public:
+	TSharedRef<FPerforceConnection> Perforce;
+
+	FFindStreamsTask(
+		TSharedRef<FPerforceConnection> InPerforce,
+		TSharedRef<FLineBasedTextWriter> InLog,
+		TArray<FString>& OutStreams,
+		const FString& StreamFilter);
+	virtual ~FFindStreamsTask() override;
+
+	virtual TSharedRef<FModalTaskResult> Run(FEvent* AbortEvent) override;
+
+private:
+	TSharedRef<FLineBasedTextWriter> Log;
+
+	TArray<FString>& Streams;
+	FString Filter;
+
+	TSharedRef<FModalTaskResult> RunInternal(FEvent* AbortEvent);
+};
+
+} // namespace UGSCore
