Fixed const-correctness and added TEXT() macro to a string literal

diff --git a/Engine/Source/Programs/SlateUGS/Private/UGSTab.cpp b/Engine/Source/Programs/SlateUGS/Private/UGSTab.cpp
--- a/Engine/Source/Programs/SlateUGS/Private/UGSTab.cpp
+++ b/Engine/Source/Programs/SlateUGS/Private/UGSTab.cpp
@@ -204,13 +204,13 @@ bool UGSTab::ShouldSyncPrecompiledEditor() const
 	return UserSettings->bSyncPrecompiledEditor && PerforceMonitor->HasZippedBinaries();
 }
 
-TArray<FString> UGSTab::GetAllStreamNames()
+TArray<FString> UGSTab::GetAllStreamNames() const
 {
 	TArray<FString> Result;
 
 	const TSharedRef<UGSCore::FModalTaskResult> TaskResult = ExecuteModalTask(
 		FSlateApplication::Get().GetActiveModalWindow(),
-		MakeShared<UGSCore::FFindStreamsTask>(PerforceClient.ToSharedRef(), MakeShared<FLineWriter>(), Result, "//*/*"),
+		MakeShared<UGSCore::FFindStreamsTask>(PerforceClient.ToSharedRef(), MakeShared<FLineWriter>(), Result, TEXT("//*/*")),
 		LOCTEXT("OpeningProjectTitle", "Finding Streams"),
 		LOCTEXT("OpeningProjectCaption", "Finding streams, please wait..."));
 	
diff --git a/Engine/Source/Programs/SlateUGS/Private/UGSTab.h b/Engine/Source/Programs/SlateUGS/Private/UGSTab.h
--- a/Engine/Source/Programs/SlateUGS/Private/UGSTab.h
+++ b/Engine/Source/Programs/SlateUGS/Private/UGSTab.h
@@ -72,7 +72,7 @@ public:
 	UGSTabManager* GetTabManager() const;
 	TSharedPtr<UGSCore::FUserSettings> GetUserSettings() const;
 	bool ShouldSyncPrecompiledEditor() const;
-	TArray<FString> GetAllStreamNames();
+	TArray<FString> GetAllStreamNames() const;
 
 	void UpdateGameTabBuildList();
 	void RefreshBuildList();
