Makes the browse button on the workspace window remember the previously opened folder

diff --git a/Engine/Source/Programs/SlateUGS/Private/Widgets/SWorkspaceWindow.cpp b/Engine/Source/Programs/SlateUGS/Private/Widgets/SWorkspaceWindow.cpp
--- a/Engine/Source/Programs/SlateUGS/Private/Widgets/SWorkspaceWindow.cpp
+++ b/Engine/Source/Programs/SlateUGS/Private/Widgets/SWorkspaceWindow.cpp
@@ -235,30 +235,21 @@ FReply SWorkspaceWindow::OnCancelClicked()
 
 FReply SWorkspaceWindow::OnBrowseClicked()
 {
-	IDesktopPlatform* DesktopPlatform = FDesktopPlatformModule::Get();
 	TArray<FString> OutOpenFilenames;
-	if (DesktopPlatform)
-	{
-		DesktopPlatform->OpenFileDialog(
-			FSlateApplication::Get().FindBestParentWindowHandleForDialogs(AsShared()),
-			LOCTEXT("OpenDialogTitle", "Open Unreal Project").ToString(),
-			TEXT(""), // Todo: maybe give a default path (such as previously selected path)
-			TEXT(""),
-			TEXT("Unreal Project Files (*.uproject)|*.uproject"),
-			EFileDialogFlags::None,
-			OutOpenFilenames
-		);
-
-		if (!OutOpenFilenames.IsEmpty())
-		{
-			FString Path = OutOpenFilenames[0];
-			if (FPaths::IsRelative(Path))
-			{
-				Path = FPaths::ConvertRelativePathToFull(Path);
-			}
+	FDesktopPlatformModule::Get()->OpenFileDialog(
+		FSlateApplication::Get().FindBestParentWindowHandleForDialogs(AsShared()),
+		LOCTEXT("OpenDialogTitle", "Open Unreal Project").ToString(),
+		PreviousProjectPath,
+		TEXT(""),
+		TEXT("Unreal Project Files (*.uproject)|*.uproject"),
+		EFileDialogFlags::None,
+		OutOpenFilenames
+	);
 
-			LocalFileText->SetText(FText::FromString(Path));
-		}
+	if (!OutOpenFilenames.IsEmpty())
+	{
+		PreviousProjectPath = OutOpenFilenames[0];
+		LocalFileText->SetText(FText::FromString(FPaths::ConvertRelativePathToFull(PreviousProjectPath)));
 	}
 
 	return FReply::Handled();
diff --git a/Engine/Source/Programs/SlateUGS/Private/Widgets/SWorkspaceWindow.h b/Engine/Source/Programs/SlateUGS/Private/Widgets/SWorkspaceWindow.h
--- a/Engine/Source/Programs/SlateUGS/Private/Widgets/SWorkspaceWindow.h
+++ b/Engine/Source/Programs/SlateUGS/Private/Widgets/SWorkspaceWindow.h
@@ -20,6 +20,8 @@ private:
 	FReply OnCancelClicked();
 
 	FReply OnBrowseClicked();
+	FString PreviousProjectPath;
+
 	FReply OnNewClicked();
 
 	bool bIsLocalFileSelected = true;
