Fix Slate UGC deprecation warning, following CL 22221537.

diff --git a/Engine/Source/Programs/SlateUGS/Private/Widgets/SLogWidget.cpp b/Engine/Source/Programs/SlateUGS/Private/Widgets/SLogWidget.cpp
--- a/Engine/Source/Programs/SlateUGS/Private/Widgets/SLogWidget.cpp
+++ b/Engine/Source/Programs/SlateUGS/Private/Widgets/SLogWidget.cpp
@@ -89,7 +89,6 @@ void SLogWidget::Construct(const FArguments& InArgs)
 		[
 			SAssignNew(MessagesTextBox, SMultiLineEditableTextBox)
 				.Style(FAppStyle::Get(), "Log.TextBox")
-				.TextStyle(FAppStyle::Get(), "Log.Normal")
 				.Padding(10.0f)
 				.ForegroundColor(FLinearColor::Gray)
 				.Marshaller(MessagesTextMarshaller)
