Fix include errors in SlateUGS for BuildSettingsVersion.V3

#rnx

diff --git a/Engine/Source/Programs/SlateUGS/Private/ChangeInfo.h b/Engine/Source/Programs/SlateUGS/Private/ChangeInfo.h
--- a/Engine/Source/Programs/SlateUGS/Private/ChangeInfo.h
+++ b/Engine/Source/Programs/SlateUGS/Private/ChangeInfo.h
@@ -2,8 +2,8 @@
 
 #pragma once
 
-#include "UGSCore/EventMonitor.h"
-#include "UGSCore/PerforceMonitor.h"
+#include "EventMonitor.h"
+#include "PerforceMonitor.h"
 
 struct FChangeInfo
 {
diff --git a/Engine/Source/Programs/SlateUGS/Private/UGSTab.cpp b/Engine/Source/Programs/SlateUGS/Private/UGSTab.cpp
--- a/Engine/Source/Programs/SlateUGS/Private/UGSTab.cpp
+++ b/Engine/Source/Programs/SlateUGS/Private/UGSTab.cpp
@@ -8,12 +8,12 @@
 #include "FindStreamsTask.h"
 #include "UGSTabManager.h"
 
-#include "UGSCore/Utility.h"
-#include "UGSCore/BuildStep.h"
-#include "UGSCore/DetectProjectSettingsTask.h"
-#include "UGSCore/FindStreamsTask.h"
-#include "UGSCore/PerforceMonitor.h"
-#include "UGSCore/EventMonitor.h"
+#include "Utility.h"
+#include "BuildStep.h"
+#include "DetectProjectSettingsTask.h"
+#include "FindStreamsTask.h"
+#include "PerforceMonitor.h"
+#include "EventMonitor.h"
 
 #include "Widgets/SModalTaskWindow.h"
 #include "Widgets/SLogWidget.h"
diff --git a/Engine/Source/Programs/SlateUGS/Private/UGSTab.h b/Engine/Source/Programs/SlateUGS/Private/UGSTab.h
--- a/Engine/Source/Programs/SlateUGS/Private/UGSTab.h
+++ b/Engine/Source/Programs/SlateUGS/Private/UGSTab.h
@@ -7,8 +7,8 @@
 #include "Widgets/SGameSyncTab.h"
 #include "Widgets/Docking/SDockTab.h"
 
-#include "UGSCore/Workspace.h"
-#include "UGSCore/UserSettings.h"
+#include "Workspace.h"
+#include "UserSettings.h"
 
 namespace UGSCore
 {
diff --git a/Engine/Source/Programs/SlateUGS/Private/UGSTabManager.cpp b/Engine/Source/Programs/SlateUGS/Private/UGSTabManager.cpp
--- a/Engine/Source/Programs/SlateUGS/Private/UGSTabManager.cpp
+++ b/Engine/Source/Programs/SlateUGS/Private/UGSTabManager.cpp
@@ -5,7 +5,7 @@
 #include "UGSLog.h"
 #include "HAL/FileManager.h"
 
-#include "UGSCore/UserSettings.h"
+#include "UserSettings.h"
 
 UGSTabManager::UGSTabManager()
 {
diff --git a/Engine/Source/Programs/SlateUGS/Private/Widgets/SLogWidget.h b/Engine/Source/Programs/SlateUGS/Private/Widgets/SLogWidget.h
--- a/Engine/Source/Programs/SlateUGS/Private/Widgets/SLogWidget.h
+++ b/Engine/Source/Programs/SlateUGS/Private/Widgets/SLogWidget.h
@@ -6,7 +6,7 @@
 #include "Framework/Text/BaseTextLayoutMarshaller.h"
 #include "Widgets/Input/SMultiLineEditableTextBox.h"
 
-#include "UGSCore/OutputAdapters.h"
+#include "OutputAdapters.h"
 
 class FLogWidgetTextLayoutMarshaller : public FBaseTextLayoutMarshaller
 {
diff --git a/Engine/Source/Programs/SlateUGS/Private/Widgets/SModalTaskWindow.cpp b/Engine/Source/Programs/SlateUGS/Private/Widgets/SModalTaskWindow.cpp
--- a/Engine/Source/Programs/SlateUGS/Private/Widgets/SModalTaskWindow.cpp
+++ b/Engine/Source/Programs/SlateUGS/Private/Widgets/SModalTaskWindow.cpp
@@ -5,7 +5,7 @@
 #include "Widgets/SBoxPanel.h"
 #include "Widgets/Text/STextBlock.h"
 #include "Widgets/Images/SThrobber.h"
-#include "UGSCore/ModalTask.h"
+#include "ModalTask.h"
 
 SModalTaskWindow::SModalTaskWindow()
 {
diff --git a/Engine/Source/Programs/SlateUGS/Private/Widgets/SScheduledSyncWindow.cpp b/Engine/Source/Programs/SlateUGS/Private/Widgets/SScheduledSyncWindow.cpp
--- a/Engine/Source/Programs/SlateUGS/Private/Widgets/SScheduledSyncWindow.cpp
+++ b/Engine/Source/Programs/SlateUGS/Private/Widgets/SScheduledSyncWindow.cpp
@@ -14,7 +14,7 @@
 #include "SSimpleButton.h"
 #include "SPopupTextWindow.h"
 
-#include "UGSCore/UserSettings.h"
+#include "UserSettings.h"
 
 #define LOCTEXT_NAMESPACE "SScheduledSyncWindow"
 
diff --git a/Engine/Source/Programs/SlateUGS/Private/Widgets/SSettingsWindow.cpp b/Engine/Source/Programs/SlateUGS/Private/Widgets/SSettingsWindow.cpp
--- a/Engine/Source/Programs/SlateUGS/Private/Widgets/SSettingsWindow.cpp
+++ b/Engine/Source/Programs/SlateUGS/Private/Widgets/SSettingsWindow.cpp
@@ -3,7 +3,7 @@
 #include "SSettingsWindow.h"
 #include "UGSTab.h"
 
-#include "UGSCore/UserSettings.h"
+#include "UserSettings.h"
 
 #include "Framework/Application/SlateApplication.h"
 #include "Widgets/Input/SCheckBox.h"
diff --git a/Engine/Source/Programs/SlateUGS/Private/Widgets/SSyncFilterWindow.h b/Engine/Source/Programs/SlateUGS/Private/Widgets/SSyncFilterWindow.h
--- a/Engine/Source/Programs/SlateUGS/Private/Widgets/SSyncFilterWindow.h
+++ b/Engine/Source/Programs/SlateUGS/Private/Widgets/SSyncFilterWindow.h
@@ -5,7 +5,7 @@
 #include "Widgets/SWindow.h"
 #include "Widgets/Input/SMultiLineEditableTextBox.h"
 
-#include "UGSCore/Workspace.h"
+#include "Workspace.h"
 
 #include "SCheckBoxList.h"
 
