BuildFix for "NonUnity Compile SlateUGS Mac"

#horde 326960

diff --git a/Engine/Source/Programs/SlateUGS/Private/Widgets/SHordeBadge.h b/Engine/Source/Programs/SlateUGS/Private/Widgets/SHordeBadge.h
--- a/Engine/Source/Programs/SlateUGS/Private/Widgets/SHordeBadge.h
+++ b/Engine/Source/Programs/SlateUGS/Private/Widgets/SHordeBadge.h
@@ -5,7 +5,7 @@
 #include "Widgets/SCompoundWidget.h"
 #include "Widgets/Input/SButton.h"
 
-enum EBadgeState
+enum class EBadgeState
 {
 	Unknown,
 
