Slate UGS: Fix incorrect & ugly errors during workspace selection

#rnx
#Synced-CL 51838013

diff --git a/Engine/Source/Programs/SlateUGS/Private/UGSTab.cpp b/Engine/Source/Programs/SlateUGS/Private/UGSTab.cpp
--- a/Engine/Source/Programs/SlateUGS/Private/UGSTab.cpp
+++ b/Engine/Source/Programs/SlateUGS/Private/UGSTab.cpp
@@ -16,6 +16,7 @@
 #include "EventMonitor.h"
 
 #include "Widgets/SModalTaskWindow.h"
+#include "Widgets/SPopupTextWindow.h"
 #include "Widgets/SLogWidget.h"
 
 #include "HAL/PlatformApplicationMisc.h"
@@ -217,7 +218,11 @@ TArray<FString> UGSTab::GetAllStreamNames() const
 	
 	if (TaskResult->Failed())
 	{
-		FMessageDialog::Open(EAppMsgType::Ok, TaskResult->GetMessage());
+		FSlateApplication::Get().AddModalWindow(
+			SNew(SPopupTextWindow)
+				.TitleText(LOCTEXT("FindStreamsErrorTitle", "Error"))
+				.BodyText(TaskResult->GetMessage()),
+			FSlateApplication::Get().GetActiveModalWindow());
 		return TArray<FString>();
 	}
 
@@ -701,7 +706,11 @@ bool UGSTab::SetupWorkspace()
 
 	if (Result->Failed())
 	{
-		FMessageDialog::Open(EAppMsgType::Ok, Result->GetMessage());
+		FSlateApplication::Get().AddModalWindow(
+			SNew(SPopupTextWindow)
+				.TitleText(LOCTEXT("SetupWorkspaceErrorTitle", "Error"))
+				.BodyText(Result->GetMessage()),
+			FSlateApplication::Get().GetActiveModalWindow());
 		return false;
 	}
 
diff --git a/Engine/Source/Programs/SlateUGS/Private/Widgets/SWorkspaceWindow.cpp b/Engine/Source/Programs/SlateUGS/Private/Widgets/SWorkspaceWindow.cpp
--- a/Engine/Source/Programs/SlateUGS/Private/Widgets/SWorkspaceWindow.cpp
+++ b/Engine/Source/Programs/SlateUGS/Private/Widgets/SWorkspaceWindow.cpp
@@ -216,17 +216,18 @@ void SWorkspaceWindow::SetWorkspaceTextBox(FText Text) const
 
 FReply SWorkspaceWindow::OnOkClicked()
 {
-	bool bIsWorkspaceValid = Tab->OnWorkspaceChosen(WorkspacePathText);
-	if (bIsWorkspaceValid)
+	if (!FPaths::FileExists(WorkspacePathText))
 	{
-		FSlateApplication::Get().FindWidgetWindow(AsShared())->RequestDestroyWindow();
-	}
-	else
-	{
-		// Todo: report other errors (no workspace associated with the file)
 		FText WorkspaceTitle = FText::FromString("Error");
 		FText Error = FText::FromString("Project file does not exist, try again");
 		FSlateApplication::Get().AddModalWindow(SNew(SPopupTextWindow).TitleText(WorkspaceTitle).BodyText(Error), SharedThis(this), false);
+		return FReply::Handled();
+	}
+
+	const bool bIsWorkspaceValid = Tab->OnWorkspaceChosen(WorkspacePathText);
+	if (bIsWorkspaceValid)
+	{
+		FSlateApplication::Get().FindWidgetWindow(AsShared())->RequestDestroyWindow();
 	}
 
 	return FReply::Handled();
